Copying //UE4/Dev-Sequencer to //UE4/Dev-Main (Source: //UE4/Dev-Sequencer @ 3178529)

diff --git a/Engine/Plugins/Messaging/UdpMessaging/Source/UdpMessaging/Private/UdpMessagingClasses.cpp b/Engine/Plugins/Messaging/UdpMessaging/Source/UdpMessaging/Private/UdpMessagingClasses.cpp
--- a/Engine/Plugins/Messaging/UdpMessaging/Source/UdpMessaging/Private/UdpMessagingClasses.cpp
+++ b/Engine/Plugins/Messaging/UdpMessaging/Source/UdpMessaging/Private/UdpMessagingClasses.cpp
@@ -1,6 +1,7 @@
 // Copyright 1998-2016 Epic Games, Inc. All Rights Reserved.
 
 #include "UdpMessagingPrivatePCH.h"
+#include "UdpMessagingSettings.h"
 
 
 UUdpMessagingSettings::UUdpMessagingSettings(const FObjectInitializer& ObjectInitializer)
diff --git a/Engine/Plugins/Messaging/UdpMessaging/Source/UdpMessaging/Private/UdpMessagingModule.cpp b/Engine/Plugins/Messaging/UdpMessaging/Source/UdpMessaging/Private/UdpMessagingModule.cpp
--- a/Engine/Plugins/Messaging/UdpMessaging/Source/UdpMessaging/Private/UdpMessagingModule.cpp
+++ b/Engine/Plugins/Messaging/UdpMessaging/Source/UdpMessaging/Private/UdpMessagingModule.cpp
@@ -1,12 +1,14 @@
 // Copyright 1998-2016 Epic Games, Inc. All Rights Reserved.
 
 #include "UdpMessagingPrivatePCH.h"
-#if WITH_EDITOR
-#include "ISettingsModule.h"
-#include "ISettingsSection.h"
-#endif
 #include "ModuleInterface.h"
-#include "ModuleManager.h"
+#include "UdpMessagingSettings.h"
+#include "UdpMessageTransport.h"
+
+#if PLATFORM_DESKTOP
+	#include "UdpMessageTunnelConnection.h"
+	#include "UdpMessageTunnel.h"
+#endif
 
 
 DEFINE_LOG_CATEGORY(LogUdpMessaging);
@@ -23,7 +25,7 @@ class FUdpMessagingModule
 {
 public:
 
-	// FSelfRegisteringExec interface
+	//~ FSelfRegisteringExec interface
 
 	virtual bool Exec(UWorld* InWorld, const TCHAR* Cmd, FOutputDevice& Ar) override
 	{
@@ -486,17 +488,19 @@ private:
 private:
 
 	/** Holds the message bridge if present. */
-	IMessageBridgePtr MessageBridge;
+	TSharedPtr<IMessageBridge, ESPMode::ThreadSafe> MessageBridge;
 
 #if PLATFORM_DESKTOP
 	/** Holds the message tunnel if present. */
-	IUdpMessageTunnelPtr MessageTunnel;
+	TSharedPtr<IUdpMessageTunnel> MessageTunnel;
 #endif
 };
 
+
 void EmptyLinkFunctionForStaticInitializationUdpMessagingTests()
 {
-	// Force references to the object files containing the functions below, to prevent them being excluded by the linker when the plugin is compiled into a static library for monolithic builds.
+	// Force references to the object files containing the functions below, to prevent them being
+	// excluded by the linker when the plug-in is compiled into a static library for monolithic builds.
 	extern void EmptyLinkFunctionForStaticInitializationUdpMessageSegmenterTest();
 	EmptyLinkFunctionForStaticInitializationUdpMessageSegmenterTest();
 	extern void EmptyLinkFunctionForStaticInitializationUdpMessageTransportTest();
@@ -505,6 +509,7 @@ void EmptyLinkFunctionForStaticInitializationUdpMessagingTests()
 	EmptyLinkFunctionForStaticInitializationUdpSerializeMessageTaskTest();
 }
 
+
 IMPLEMENT_MODULE(FUdpMessagingModule, UdpMessaging);
 
 
diff --git a/Engine/Plugins/Messaging/UdpMessaging/Source/UdpMessaging/Private/UdpMessagingPrivatePCH.h b/Engine/Plugins/Messaging/UdpMessaging/Source/UdpMessaging/Private/UdpMessagingPrivatePCH.h
--- a/Engine/Plugins/Messaging/UdpMessaging/Source/UdpMessaging/Private/UdpMessagingPrivatePCH.h
+++ b/Engine/Plugins/Messaging/UdpMessaging/Source/UdpMessaging/Private/UdpMessagingPrivatePCH.h
@@ -3,19 +3,24 @@
 #pragma once
 
 
-/* Private dependencies
- *****************************************************************************/
-
-#include "Core.h"
-#include "CoreUObject.h"
-#include "Messaging.h"
-#include "Networking.h"
-#include "Sockets.h"
-#include "SocketSubsystem.h"
-
+#include "Runtime/Core/Public/Core.h"
+#include "Runtime/Core/Public/Misc/AutomationTest.h"
+#include "Runtime/Core/Public/Modules/ModuleManager.h"
+#include "Runtime/CoreUObject/Public/CoreUObject.h"
+#include "Runtime/Messaging/Public/Deprecated/Messaging.h"
+#include "Runtime/Networking/Public/Networking.h"
+#include "Runtime/Serialization/Public/StructDeserializer.h"
+#include "Runtime/Serialization/Public/StructSerializer.h"
+#include "Runtime/Serialization/Public/Backends/JsonStructDeserializerBackend.h"
+#include "Runtime/Serialization/Public/Backends/JsonStructSerializerBackend.h"
+#include "Runtime/Sockets/Public/Sockets.h"
+#include "Runtime/Sockets/Public/SocketSubsystem.h"
+
+#if WITH_EDITOR
+	#include "Developer/Settings/Public/ISettingsModule.h"
+	#include "Developer/Settings/Public/ISettingsSection.h"
+#endif
 
-/* Private constants
- *****************************************************************************/
 
 /** Declares a log category for this module. */
 DECLARE_LOG_CATEGORY_EXTERN(LogUdpMessaging, Log, All);
@@ -35,30 +40,3 @@ DECLARE_LOG_CATEGORY_EXTERN(LogUdpMessaging, Log, All);
 
 /** Defines the protocol version of the UDP message transport. */
 #define UDP_MESSAGING_TRANSPORT_PROTOCOL_VERSION 10
-
-
-/* Private includes
- *****************************************************************************/
-
-// shared
-#include "UdpMessageSegment.h"
-#include "UdpMessagingSettings.h"
-
-// transport
-#include "UdpMessageBeacon.h"
-#include "UdpReassembledMessage.h"
-#include "UdpMessageResequencer.h"
-#include "UdpDeserializedMessage.h"
-#include "UdpSerializedMessage.h"
-#include "UdpMessageSegmenter.h"
-#include "UdpMessageProcessor.h"
-#include "UdpSerializeMessageTask.h"
-#include "UdpMessageTransport.h"
-
-#if PLATFORM_DESKTOP
-
-// tunnel
-#include "UdpMessageTunnelConnection.h"
-#include "UdpMessageTunnel.h"
-
-#endif
diff --git a/Engine/Plugins/Messaging/UdpMessaging/Source/UdpMessaging/Public/IUdpMessageTunnel.h b/Engine/Plugins/Messaging/UdpMessaging/Source/UdpMessaging/Public/IUdpMessageTunnel.h
--- a/Engine/Plugins/Messaging/UdpMessaging/Source/UdpMessaging/Public/IUdpMessageTunnel.h
+++ b/Engine/Plugins/Messaging/UdpMessaging/Source/UdpMessaging/Public/IUdpMessageTunnel.h
@@ -3,7 +3,6 @@
 #pragma once
 
 
-// forward declarations
 struct FIPv4Endpoint;
 class IUdpMessageTunnelConnection;
 
@@ -86,10 +85,3 @@ public:
 	/** Virtual destructor. */
 	virtual ~IUdpMessageTunnel() { }
 };
-
-
-/** Type definition for shared pointers to instances of IUdpMessageTunnel. */
-typedef TSharedPtr<IUdpMessageTunnel> IUdpMessageTunnelPtr;
-
-/** Type definition for shared references to instances of IUdpMessageTunnel. */
-typedef TSharedRef<IUdpMessageTunnel> IUdpMessageTunnelRef;
diff --git a/Engine/Plugins/Messaging/UdpMessaging/Source/UdpMessaging/Public/IUdpMessageTunnelConnection.h b/Engine/Plugins/Messaging/UdpMessaging/Source/UdpMessaging/Public/IUdpMessageTunnelConnection.h
--- a/Engine/Plugins/Messaging/UdpMessaging/Source/UdpMessaging/Public/IUdpMessageTunnelConnection.h
+++ b/Engine/Plugins/Messaging/UdpMessaging/Source/UdpMessaging/Public/IUdpMessageTunnelConnection.h
@@ -53,10 +53,3 @@ public:
 	/** Virtual destructor. */
 	virtual ~IUdpMessageTunnelConnection() { }
 };
-
-
-/** Type definition for shared pointers to instances of IUdpMessageTunnelConnection. */
-typedef TSharedPtr<IUdpMessageTunnelConnection> IUdpMessageTunnelConnectionPtr;
-
-/** Type definition for shared references to instances of IUdpMessageTunnelConnection. */
-typedef TSharedRef<IUdpMessageTunnelConnection> IUdpMessageTunnelConnectionRef;
diff --git a/Engine/Plugins/MovieScene/LevelSequenceEditor/Source/LevelSequenceEditor/Private/CinematicViewport/SCinematicLevelViewport.cpp b/Engine/Plugins/MovieScene/LevelSequenceEditor/Source/LevelSequenceEditor/Private/CinematicViewport/SCinematicLevelViewport.cpp
--- a/Engine/Plugins/MovieScene/LevelSequenceEditor/Source/LevelSequenceEditor/Private/CinematicViewport/SCinematicLevelViewport.cpp
+++ b/Engine/Plugins/MovieScene/LevelSequenceEditor/Source/LevelSequenceEditor/Private/CinematicViewport/SCinematicLevelViewport.cpp
@@ -435,7 +435,7 @@ void SCinematicLevelViewport::SetTime(float Value)
 	ISequencer* Sequencer = GetSequencer();
 	if (Sequencer)
 	{
-		Sequencer->SetGlobalTime(Value);
+		Sequencer->SetLocalTime(Value);
 	}
 }
 
@@ -444,7 +444,7 @@ float SCinematicLevelViewport::GetTime() const
 	ISequencer* Sequencer = GetSequencer();
 	if (Sequencer)
 	{
-		return Sequencer->GetGlobalTime();
+		return Sequencer->GetLocalTime();
 	}
 	return 0.f;
 }
@@ -589,7 +589,7 @@ void SCinematicLevelViewport::Tick(const FGeometry& AllottedGeometry, const doub
 	UMovieSceneSubSection* SubSection = nullptr;
 	if (SubTrack)
 	{
-		const float CurrentTime = Sequencer->GetGlobalTime();
+		const float CurrentTime = Sequencer->GetLocalTime();
 
 		for (UMovieSceneSection* Section : SubTrack->GetAllSections())
 		{
@@ -600,7 +600,7 @@ void SCinematicLevelViewport::Tick(const FGeometry& AllottedGeometry, const doub
 		}
 	}
 
-	const float AbsoluteTime = Sequencer->GetGlobalTime();
+	const float AbsoluteTime = Sequencer->GetLocalTime();
 
 	FText TimeFormat = LOCTEXT("TimeFormat", "{0}");
 
@@ -615,8 +615,9 @@ void SCinematicLevelViewport::Tick(const FGeometry& AllottedGeometry, const doub
 			PlaybackRangeStart = ShotMovieScene->GetPlaybackRange().GetLowerBoundValue();
 		}
 
-		const float ShotOffset = SubSection->StartOffset + PlaybackRangeStart - SubSection->PrerollTime;
-		const float AbsoluteShotPosition = ShotOffset + (AbsoluteTime - (SubSection->GetStartTime() - SubSection->PrerollTime)) / SubSection->TimeScale;
+		const float InnerOffset = (AbsoluteTime - SubSection->GetStartTime()) * SubSection->Parameters.TimeScale;
+		const float AbsoluteShotPosition = PlaybackRangeStart + SubSection->Parameters.StartOffset + InnerOffset;
+
 		UIData.Frame = FText::Format(
 			TimeFormat,
 			FText::FromString(ZeroPadTypeInterface->ToString(AbsoluteShotPosition))
diff --git a/Engine/Plugins/MovieScene/LevelSequenceEditor/Source/LevelSequenceEditor/Private/CinematicViewport/SCinematicTransportRange.cpp b/Engine/Plugins/MovieScene/LevelSequenceEditor/Source/LevelSequenceEditor/Private/CinematicViewport/SCinematicTransportRange.cpp
--- a/Engine/Plugins/MovieScene/LevelSequenceEditor/Source/LevelSequenceEditor/Private/CinematicViewport/SCinematicTransportRange.cpp
+++ b/Engine/Plugins/MovieScene/LevelSequenceEditor/Source/LevelSequenceEditor/Private/CinematicViewport/SCinematicTransportRange.cpp
@@ -78,7 +78,7 @@ void SCinematicTransportRange::SetTime(const FGeometry& MyGeometry, const FPoint
 		
 		const TRange<float> WorkingRange = Sequencer->GetFocusedMovieSceneSequence()->GetMovieScene()->GetEditorData().WorkingRange;
 		
-		Sequencer->SetGlobalTime(WorkingRange.GetLowerBoundValue() + WorkingRange.Size<float>()*Lerp, ESnapTimeMode::STM_All);
+		Sequencer->SetLocalTime(WorkingRange.GetLowerBoundValue() + WorkingRange.Size<float>()*Lerp, ESnapTimeMode::STM_All);
 	}
 }
 
@@ -145,7 +145,7 @@ int32 SCinematicTransportRange::OnPaint(const FPaintArgs& Args, const FGeometry&
 		FLinearColor(MidGray)
 	);
 
-	const float CurrentTime = Sequencer->GetGlobalTime();
+	const float CurrentTime = Sequencer->GetLocalTime();
 	const float ProgressLerp = (CurrentTime - WorkingRange.GetLowerBoundValue()) / FullRange;
 
 	// Draw the playback progress
diff --git a/Engine/Plugins/MovieScene/LevelSequenceEditor/Source/LevelSequenceEditor/Private/LevelSequenceEditorToolkit.cpp b/Engine/Plugins/MovieScene/LevelSequenceEditor/Source/LevelSequenceEditor/Private/LevelSequenceEditorToolkit.cpp
--- a/Engine/Plugins/MovieScene/LevelSequenceEditor/Source/LevelSequenceEditor/Private/LevelSequenceEditorToolkit.cpp
+++ b/Engine/Plugins/MovieScene/LevelSequenceEditor/Source/LevelSequenceEditor/Private/LevelSequenceEditorToolkit.cpp
@@ -29,7 +29,6 @@
 #include "SDockTab.h"
 #include "LevelSequencePlayer.h"
 #include "SequencerSettings.h"
-#include "SequencerSpawnRegister.h"
 #include "MovieSceneCaptureDialogModule.h"
 
 // @todo sequencer: hack: setting defaults for transform tracks
@@ -197,7 +196,7 @@ void FLevelSequenceEditorToolkit::Initialize(const EToolkitMode::Type Mode, cons
 
 	FAssetEditorToolkit::InitAssetEditor(Mode, InitToolkitHost, SequencerDefs::SequencerAppIdentifier, StandaloneDefaultLayout, bCreateDefaultStandaloneMenu, bCreateDefaultToolbar, LevelSequence);
 
-	TSharedRef<FLevelSequenceEditorSpawnRegister> SpawnRegister = MakeShareable(new TSequencerSpawnRegister<FLevelSequenceEditorSpawnRegister>);
+	TSharedRef<FLevelSequenceEditorSpawnRegister> SpawnRegister = MakeShareable(new FLevelSequenceEditorSpawnRegister);
 
 	// initialize sequencer
 	FSequencerInitParams SequencerInitParams;
@@ -686,6 +685,7 @@ void FLevelSequenceEditorToolkit::AddShot(UMovieSceneCinematicShotTrack* ShotTra
 
 	// Focus on the new shot
 	GetSequencer()->UpdateRuntimeInstances();
+	GetSequencer()->ForceEvaluate();
 	GetSequencer()->FocusSequenceInstance(*ShotSubSection);
 
 	const ULevelSequenceMasterSequenceSettings* MasterSequenceSettings = GetDefault<ULevelSequenceMasterSequenceSettings>();
diff --git a/Engine/Plugins/MovieScene/LevelSequenceEditor/Source/LevelSequenceEditor/Private/Misc/LevelSequenceEditorSpawnRegister.cpp b/Engine/Plugins/MovieScene/LevelSequenceEditor/Source/LevelSequenceEditor/Private/Misc/LevelSequenceEditorSpawnRegister.cpp
--- a/Engine/Plugins/MovieScene/LevelSequenceEditor/Source/LevelSequenceEditor/Private/Misc/LevelSequenceEditorSpawnRegister.cpp
+++ b/Engine/Plugins/MovieScene/LevelSequenceEditor/Source/LevelSequenceEditor/Private/Misc/LevelSequenceEditorSpawnRegister.cpp
@@ -6,6 +6,8 @@
 #include "Editor/LevelEditor/Public/LevelEditor.h"
 #include "MovieScene.h"
 #include "AssetSelection.h"
+#include "Evaluation/MovieSceneSpawnTemplate.h"
+#include "Evaluation/MovieSceneEvaluationTemplateInstance.h"
 
 #define LOCTEXT_NAMESPACE "LevelSequenceEditorSpawnRegister"
 
@@ -53,18 +55,18 @@ FLevelSequenceEditorSpawnRegister::~FLevelSequenceEditorSpawnRegister()
 /* FLevelSequenceSpawnRegister interface
  *****************************************************************************/
 
-UObject* FLevelSequenceEditorSpawnRegister::SpawnObject(const FGuid& BindingId, FMovieSceneSequenceInstance& SequenceInstance, IMovieScenePlayer& Player)
+UObject* FLevelSequenceEditorSpawnRegister::SpawnObject(FMovieSceneSpawnable& Spawnable, FMovieSceneSequenceIDRef TemplateID, IMovieScenePlayer& Player)
 {
 	TGuardValue<bool> Guard(bShouldClearSelectionCache, false);
-	AActor* NewObject = Cast<AActor>(FLevelSequenceSpawnRegister::SpawnObject(BindingId, SequenceInstance, Player));
+	AActor* NewObject = Cast<AActor>(FLevelSequenceSpawnRegister::SpawnObject(Spawnable, TemplateID, Player));
 	
 	if (NewObject)
 	{
 		// Cache the spawned object, and editable state first
-		SpawnedObjects.FindOrAdd(SequenceInstance.GetInstanceId()).Add(NewObject);
+		SpawnedObjects.FindOrAdd(TemplateID).Add(NewObject);
 
 		// Select the actor if we think it should be selected
-		if (SelectedSpawnedObjects.Contains(FMovieSceneSpawnRegisterKey(BindingId, SequenceInstance)))
+		if (SelectedSpawnedObjects.Contains(FMovieSceneSpawnRegisterKey(TemplateID, Spawnable.GetGuid())))
 		{
 			GEditor->SelectActor(NewObject, true /*bSelected*/, true /*bNotify*/);
 		}
@@ -74,60 +76,89 @@ UObject* FLevelSequenceEditorSpawnRegister::SpawnObject(const FGuid& BindingId,
 }
 
 
-void FLevelSequenceEditorSpawnRegister::PreDestroyObject(UObject& Object, const FGuid& BindingId, FMovieSceneSequenceInstance& SequenceInstance)
+void FLevelSequenceEditorSpawnRegister::PreDestroyObject(UObject& Object, const FGuid& BindingId, FMovieSceneSequenceIDRef TemplateID)
 {
 	TGuardValue<bool> Guard(bShouldClearSelectionCache, false);
 	
 	TSharedPtr<ISequencer> Sequencer = WeakSequencer.Pin();
 
 	// We only save default state for the currently focussed movie scene sequence instance
-	if (Sequencer.IsValid() && &Sequencer->GetFocusedMovieSceneSequenceInstance().Get() == &SequenceInstance)
+	if (Sequencer.IsValid() && Sequencer->GetFocusedTemplateID() == TemplateID)
 	{
-		SaveDefaultSpawnableState(BindingId, SequenceInstance);
+		SaveDefaultSpawnableState(BindingId);
 	}
 
 	// Cache its selection state
 	AActor* Actor = Cast<AActor>(&Object);
 	if (Actor && GEditor->GetSelectedActors()->IsSelected(Actor))
 	{
-		SelectedSpawnedObjects.Add(FMovieSceneSpawnRegisterKey(BindingId, SequenceInstance));
+		SelectedSpawnedObjects.Add(FMovieSceneSpawnRegisterKey(TemplateID, BindingId));
 		GEditor->SelectActor(Actor, false /*bSelected*/, true /*bNotify*/);
 	}
 	
 	// Remove the spawned object (and anything that's null) from our cache
-	TSet<FObjectKey>* ExistingObjects = SpawnedObjects.Find(SequenceInstance.GetInstanceId());
+	TSet<FObjectKey>* ExistingObjects = SpawnedObjects.Find(TemplateID);
 	if (ExistingObjects)
 	{
 		ExistingObjects->Remove(FObjectKey(&Object));
 		if (ExistingObjects->Num() == 0)
 		{
-			SpawnedObjects.Remove(SequenceInstance.GetInstanceId());
+			SpawnedObjects.Remove(TemplateID);
 		}
 	}
 
-	FLevelSequenceSpawnRegister::PreDestroyObject(Object, BindingId, SequenceInstance);
+	FLevelSequenceSpawnRegister::PreDestroyObject(Object, BindingId, TemplateID);
 }
 
-void FLevelSequenceEditorSpawnRegister::SaveDefaultSpawnableState(const FGuid& BindingId, FMovieSceneSequenceInstance& SequenceInstance)
+void FLevelSequenceEditorSpawnRegister::SaveDefaultSpawnableState(FMovieSceneSpawnable& Spawnable, FMovieSceneSequenceIDRef TemplateID, IMovieScenePlayer& Player)
 {
-	TSharedPtr<ISequencer> Sequencer = WeakSequencer.Pin();
-	if (!Sequencer.IsValid())
+	const FMovieSceneEvaluationTemplateInstance* Template = Player.GetEvaluationTemplate().GetInstance(TemplateID);
+	UMovieSceneSequence* Sequence = Template ? Template->Sequence : nullptr;
+
+	UObject* Object = FindSpawnedObject(Spawnable.GetGuid(), TemplateID);
+	if (!Object || !Sequence)
 	{
 		return;
 	}
 
-	UObject* Object = SequenceInstance.FindObject(BindingId, *Sequencer);
-	if (!Object)
+	auto RestorePredicate = [](FMovieSceneAnimTypeID TypeID){ return TypeID != FMovieSceneSpawnSectionTemplate::GetAnimTypeID(); };
+
+	if (AActor* Actor = Cast<AActor>(Object))
+	{
+		// Restore state on any components
+		for (UActorComponent* Component : TInlineComponentArray<UActorComponent*>(Actor))
+		{
+			if (Component)
+			{
+				Player.RestorePreAnimatedState(*Component, RestorePredicate);
+			}
+		}
+	}
+
+	// Restore state on the object itself
+	Player.RestorePreAnimatedState(*Object, RestorePredicate);
+
+	// Copy the template
+	Spawnable.CopyObjectTemplate(*Object, *Sequence);
+}
+
+void FLevelSequenceEditorSpawnRegister::SaveDefaultSpawnableState(const FGuid& BindingId)
+{
+	TSharedPtr<ISequencer> Sequencer = WeakSequencer.Pin();
+	if (!Sequencer.IsValid())
 	{
 		return;
 	}
 
-	// Find the spawnable definition
-	FMovieSceneSpawnable* Spawnable = SequenceInstance.GetSequence()->GetMovieScene()->FindSpawnable(BindingId);
+	UMovieSceneSequence* Sequence = Sequencer->GetFocusedMovieSceneSequence();
+	UMovieScene* MovieScene = Sequence ? Sequence->GetMovieScene() : nullptr;
+	FMovieSceneSequenceID TemplateID = Sequencer->GetFocusedTemplateID();
+
+	FMovieSceneSpawnable* Spawnable = MovieScene->FindSpawnable(BindingId);
+
 	if (Spawnable)
 	{
-		SequenceInstance.RestoreSpecificState(BindingId, *Sequencer);
-		Spawnable->CopyObjectTemplate(*Object, *SequenceInstance.GetSequence());
+		SaveDefaultSpawnableState(*Spawnable, TemplateID, *Sequencer);
 	}
 }
 
@@ -135,8 +166,7 @@ void FLevelSequenceEditorSpawnRegister::OnPreSaveMovieScene(ISequencer& InSequen
 {
 	// We're about to save the movie scene(s), so we need to save default spawnable state for the currently focused movie scene sequence instance
 
-	TSharedPtr<FMovieSceneSequenceInstance> Instance = InSequencer.GetFocusedMovieSceneSequenceInstance();
-	UMovieSceneSequence* Sequence = Instance->GetSequence();
+	UMovieSceneSequence* Sequence = InSequencer.GetFocusedMovieSceneSequence();
 	UMovieScene* MovieScene = Sequence ? Sequence->GetMovieScene() : nullptr;
 	if (!MovieScene)
 	{
@@ -146,13 +176,13 @@ void FLevelSequenceEditorSpawnRegister::OnPreSaveMovieScene(ISequencer& InSequen
 	for (int32 SpawnableIndex = 0; SpawnableIndex < MovieScene->GetSpawnableCount(); ++SpawnableIndex)
 	{
 		FMovieSceneSpawnable& Spawnable = MovieScene->GetSpawnable(SpawnableIndex);
-		SaveDefaultSpawnableState(Spawnable.GetGuid(), *Instance);
+		SaveDefaultSpawnableState(Spawnable.GetGuid());
 	}
 }
 
-void FLevelSequenceEditorSpawnRegister::OnSequenceInstanceActivated(FMovieSceneSequenceInstance& ActiveInstance)
+void FLevelSequenceEditorSpawnRegister::OnSequenceInstanceActivated(FMovieSceneSequenceIDRef InTemplateID)
 {
-	ActiveSequence = ActiveInstance.GetInstanceId();
+	ActiveSequence = InTemplateID;
 }
 
 /* FLevelSequenceEditorSpawnRegister implementation
@@ -164,7 +194,7 @@ void FLevelSequenceEditorSpawnRegister::SetSequencer(const TSharedPtr<ISequencer
 	Sequencer->OnPreSave().AddRaw(this, &FLevelSequenceEditorSpawnRegister::OnPreSaveMovieScene);
 	Sequencer->OnActivateSequence().AddRaw(this, &FLevelSequenceEditorSpawnRegister::OnSequenceInstanceActivated);
 
-	ActiveSequence = Sequencer->GetFocusedMovieSceneSequenceInstance()->GetInstanceId();
+	ActiveSequence = Sequencer->GetFocusedTemplateID();
 }
 
 
@@ -175,19 +205,22 @@ void FLevelSequenceEditorSpawnRegister::HandleActorSelectionChanged(const TArray
 {
 	if (bShouldClearSelectionCache)
 	{
-		TSharedPtr<ISequencer> Sequencer = WeakSequencer.Pin();
-		if (Sequencer.IsValid())
-		{
-			for (auto SpawnedObject : SpawnedObjects)
-			{
-				FMovieSceneSpawnRegisterKey SpawnRegisterKey(SpawnedObject.Key, Sequencer->GetFocusedMovieSceneSequenceInstance().Get());
-				if (SelectedSpawnedObjects.Contains(SpawnRegisterKey))
-				{
-					SelectedSpawnedObjects.Remove(SpawnRegisterKey);
-					break;
-				}
-			}
-		}
+		// @todo: arodham: how does this merge? SpawnedObject.Key relates to the old unique sequence instance ID, not the object binding ID. Not sure what's intended here. Source Ref: CL#3117184
+		// TSharedPtr<ISequencer> Sequencer = WeakSequencer.Pin();
+		// if (Sequencer.IsValid())
+		// {
+		// 	for (auto SpawnedObject : SpawnedObjects)
+		// 	{
+		// 		FMovieSceneSpawnRegisterKey SpawnRegisterKey(SpawnedObject.Key, Sequencer->GetFocusedMovieSceneSequenceInstance().Get());
+		// 		if (SelectedSpawnedObjects.Contains(SpawnRegisterKey))
+		// 		{
+		// 			SelectedSpawnedObjects.Remove(SpawnRegisterKey);
+		// 			break;
+		// 		}
+		// 	}
+		// }
+
+		SelectedSpawnedObjects.Reset();
 	}
 }
 
@@ -243,14 +276,14 @@ void FLevelSequenceEditorSpawnRegister::OnObjectsReplaced(const TMap<UObject*, U
 		UObject* SpawnedObject = WeakObject.Get();
 		if (UObject* NewObject = OldToNewInstanceMap.FindRef(SpawnedObject))
 		{
+			// Reassign the object
 			WeakObject = NewObject;
+
 			// It's a spawnable, so ensure it's transient
 			NewObject->SetFlags(RF_Transient);
-			TSharedPtr<FMovieSceneSequenceInstance> Instance = Pair.Key.SequenceInstance.Pin();
-			if (Instance.IsValid())
-			{
-				Instance->OnObjectSpawned(Pair.Key.BindingId, *NewObject, *Sequencer);
-			}
+
+			// Invalidate the binding - it will be resolved if it's ever asked for again
+			Sequencer->State.Invalidate(Pair.Key.BindingId, Pair.Key.TemplateID);
 		}
 	}
 }
diff --git a/Engine/Plugins/MovieScene/LevelSequenceEditor/Source/LevelSequenceEditor/Private/Misc/LevelSequenceEditorSpawnRegister.h b/Engine/Plugins/MovieScene/LevelSequenceEditor/Source/LevelSequenceEditor/Private/Misc/LevelSequenceEditorSpawnRegister.h
--- a/Engine/Plugins/MovieScene/LevelSequenceEditor/Source/LevelSequenceEditor/Private/Misc/LevelSequenceEditorSpawnRegister.h
+++ b/Engine/Plugins/MovieScene/LevelSequenceEditor/Source/LevelSequenceEditor/Private/Misc/LevelSequenceEditorSpawnRegister.h
@@ -6,7 +6,6 @@
 #include "ObjectKey.h"
 
 class AActor;
-class FMovieSceneSequenceInstance;
 class IMovieScenePlayer;
 class ISequencer;
 class UObject;
@@ -35,8 +34,9 @@ public:
 
 	// FLevelSequenceSpawnRegister interface
 
-	virtual UObject* SpawnObject(const FGuid& BindingId, FMovieSceneSequenceInstance& SequenceInstance, IMovieScenePlayer& Player) override;
-	virtual void PreDestroyObject(UObject& Object, const FGuid& BindingId, FMovieSceneSequenceInstance& SequenceInstance) override;
+	virtual UObject* SpawnObject(FMovieSceneSpawnable& Spawnable, FMovieSceneSequenceIDRef TemplateID, IMovieScenePlayer& Player) override;
+	virtual void PreDestroyObject(UObject& Object, const FGuid& BindingId, FMovieSceneSequenceIDRef TemplateID) override;
+	virtual void SaveDefaultSpawnableState(FMovieSceneSpawnable& Spawnable, FMovieSceneSequenceIDRef TemplateID, IMovieScenePlayer& Player) override;
 #if WITH_EDITOR
 	virtual TValueOrError<FNewSpawnable, FText> CreateNewSpawnableType(UObject& SourceObject, UMovieScene& OwnerMovieScene) override;
 #endif
@@ -50,10 +50,10 @@ private:
 	void OnPreSaveMovieScene(ISequencer& InSequencer);
 
 	/** Called when a new movie scene sequence instance has been activated */
-	void OnSequenceInstanceActivated(FMovieSceneSequenceInstance& ActiveInstance);
+	void OnSequenceInstanceActivated(FMovieSceneSequenceIDRef ActiveInstance);
 
 	/** Saves the default state for the specified spawnable, if an instance for it currently exists */
-	void SaveDefaultSpawnableState(const FGuid& BindingId, FMovieSceneSequenceInstance& SequenceInstance);
+	void SaveDefaultSpawnableState(const FGuid& BindingId);
 
 	/** Check whether the specified objects are editable on the details panel. Called from the level editor */
 	bool AreObjectsEditable(const TArray<TWeakObjectPtr<UObject>>& InObjects) const;
@@ -70,7 +70,7 @@ private:
 	TSet<FMovieSceneSpawnRegisterKey> SelectedSpawnedObjects;
 
 	/** Set of currently spawned objects */
-	TMap<FGuid, TSet<FObjectKey>> SpawnedObjects;
+	TMap<FMovieSceneSequenceID, TSet<FObjectKey>> SpawnedObjects;
 
 	/** True if we should clear the above selection cache when the editor selection has been changed. */
 	bool bShouldClearSelectionCache;
@@ -79,5 +79,5 @@ private:
 	TWeakPtr<ISequencer> WeakSequencer;
 
 	/** Identifier for the current active level sequence */
-	FGuid ActiveSequence;
+	FMovieSceneSequenceID ActiveSequence;
 };
diff --git a/Engine/Plugins/MovieScene/MatineeToLevelSequence/Source/MatineeToLevelSequence/Private/MatineeToLevelSequenceModule.cpp b/Engine/Plugins/MovieScene/MatineeToLevelSequence/Source/MatineeToLevelSequence/Private/MatineeToLevelSequenceModule.cpp
--- a/Engine/Plugins/MovieScene/MatineeToLevelSequence/Source/MatineeToLevelSequence/Private/MatineeToLevelSequenceModule.cpp
+++ b/Engine/Plugins/MovieScene/MatineeToLevelSequence/Source/MatineeToLevelSequence/Private/MatineeToLevelSequenceModule.cpp
@@ -4,6 +4,8 @@
 #include "ModuleInterface.h"
 #include "NotificationManager.h"
 #include "SNotificationList.h"
+#include "Matinee/InterpTrack.h"
+#include "MatineeToLevelSequenceModule.h"
 
 #define LOCTEXT_NAMESPACE "MatineeToLevelSequence"
 
@@ -13,17 +15,16 @@ DEFINE_LOG_CATEGORY(LogMatineeToLevelSequence);
  * Implements the MatineeToLevelSequence module.
  */
 class FMatineeToLevelSequenceModule
-	: public IModuleInterface
+	: public IMatineeToLevelSequenceModule
 {
 public:
-
 	// IModuleInterface interface
 
 	virtual void StartupModule() override
 	{
 		if (GEditor)
 		{
-			GEditor->OnShouldOpenMatinee().BindStatic(ShouldOpenMatinee);
+			GEditor->OnShouldOpenMatinee().BindRaw(this, &FMatineeToLevelSequenceModule::ShouldOpenMatinee);
 		}
 		
 		RegisterMenuExtensions();
@@ -34,13 +35,38 @@ public:
 		UnregisterMenuExtensions();
 	}
 
+ 	FDelegateHandle RegisterTrackConverterForMatineeClass(TSubclassOf<UInterpTrack> InterpTrackClass, FOnConvertMatineeTrack OnConvertMatineeTrack)
+	{
+		if (ExtendedInterpConverters.Contains(InterpTrackClass))
+		{
+			UE_LOG(LogMatineeToLevelSequence, Warning, TEXT("Track converter already registered for: %s"), InterpTrackClass->GetClass());
+			return FDelegateHandle();
+		}
+
+		return ExtendedInterpConverters.Add(InterpTrackClass, OnConvertMatineeTrack).GetHandle();
+	}
+ 	
+	void UnregisterTrackConverterForMatineeClass(FDelegateHandle RemoveDelegate)
+	{
+		for (auto InterpConverter : ExtendedInterpConverters)
+		{
+			if (InterpConverter.Value.GetHandle() == RemoveDelegate)
+			{
+				ExtendedInterpConverters.Remove(*InterpConverter.Key);
+				return;
+			}
+		}
+
+		UE_LOG(LogMatineeToLevelSequence, Warning, TEXT("Attempted to remove track convert that could not be found"));
+	}
+
 protected:
 
 	/** Register menu extensions for the level editor toolbar. */
 	void RegisterMenuExtensions()
 	{
 		// Register level editor menu extender
-		LevelEditorMenuExtenderDelegate = FLevelEditorModule::FLevelViewportMenuExtender_SelectedActors::CreateStatic(&FMatineeToLevelSequenceModule::ExtendLevelViewportContextMenu);
+		LevelEditorMenuExtenderDelegate = FLevelEditorModule::FLevelViewportMenuExtender_SelectedActors::CreateRaw(this, &FMatineeToLevelSequenceModule::ExtendLevelViewportContextMenu);
 		FLevelEditorModule& LevelEditorModule = FModuleManager::LoadModuleChecked<FLevelEditorModule>(TEXT("LevelEditor"));
 		auto& MenuExtenders = LevelEditorModule.GetAllLevelViewportContextMenuExtenders();
 		MenuExtenders.Add(LevelEditorMenuExtenderDelegate);
@@ -62,7 +88,7 @@ protected:
 		}
 	}
 
-	static TSharedRef<FExtender> ExtendLevelViewportContextMenu(const TSharedRef<FUICommandList> CommandList, const TArray<AActor*> SelectedActors)
+	TSharedRef<FExtender> ExtendLevelViewportContextMenu(const TSharedRef<FUICommandList> CommandList, const TArray<AActor*> SelectedActors)
 	{
 		TSharedRef<FExtender> Extender(new FExtender());
 
@@ -82,13 +108,13 @@ protected:
 				"ActorSelectVisibilityLevels",
 				EExtensionHook::After,
 				nullptr,
-				FMenuExtensionDelegate::CreateStatic(&FMatineeToLevelSequenceModule::CreateLevelViewportContextMenuEntries, ActorsToConvert));
+				FMenuExtensionDelegate::CreateRaw(this, &FMatineeToLevelSequenceModule::CreateLevelViewportContextMenuEntries, ActorsToConvert));
 		}
 
 		return Extender;
 	}
 
-	static void CreateLevelViewportContextMenuEntries(FMenuBuilder& MenuBuilder, TArray<TWeakObjectPtr<AActor> > ActorsToConvert)
+	void CreateLevelViewportContextMenuEntries(FMenuBuilder& MenuBuilder, TArray<TWeakObjectPtr<AActor> > ActorsToConvert)
 	{
 		MenuBuilder.BeginSection("LevelSequence", LOCTEXT("LevelSequenceLevelEditorHeading", "Level Sequence"));
 
@@ -96,7 +122,7 @@ protected:
 			LOCTEXT("MenuExtensionConvertMatineeToLevelSequence", "Convert to Level Sequence"),
 			LOCTEXT("MenuExtensionConvertMatineeToLevelSequence_Tooltip", "Convert to Level Sequence"),
 			FSlateIcon(),
-			FExecuteAction::CreateStatic(&FMatineeToLevelSequenceModule::OnConvertMatineeToLevelSequence, ActorsToConvert),
+			FExecuteAction::CreateRaw(this, &FMatineeToLevelSequenceModule::OnConvertMatineeToLevelSequence, ActorsToConvert),
 			NAME_None,
 			EUserInterfaceActionType::Button);
 
@@ -104,7 +130,7 @@ protected:
 	}
 
 	/** Callback when opening a matinee. Prompts the user whether to convert this matinee to a level sequence actor */
-	static bool ShouldOpenMatinee(AMatineeActor* MatineeActor)
+	bool ShouldOpenMatinee(AMatineeActor* MatineeActor)
 	{
 		//@todo Camera anims aren't supported as level sequence assets yet
 		if (MatineeActor->IsA(AMatineeActorCameraAnim::StaticClass()))
@@ -137,7 +163,7 @@ protected:
 	}
 
 	/** Callback for converting a matinee to a level sequence asset. */
-	static void OnConvertMatineeToLevelSequence(TArray<TWeakObjectPtr<AActor> > ActorsToConvert)
+	void OnConvertMatineeToLevelSequence(TArray<TWeakObjectPtr<AActor> > ActorsToConvert)
 	{
 		TArray<TWeakObjectPtr<ALevelSequenceActor> > NewActors;
 
@@ -197,6 +223,30 @@ protected:
 		}
 	}
 
+	/** Find or add a folder for the given actor **/
+	static UMovieSceneFolder* FindOrAddFolder(UMovieScene* MovieScene, FName FolderName)
+	{
+		// look for a folder to put us in
+		UMovieSceneFolder* FolderToUse = nullptr;
+		for (UMovieSceneFolder* Folder : MovieScene->GetRootFolders())
+		{
+			if (Folder->GetFolderName() == FolderName)
+			{
+				FolderToUse = Folder;
+				break;
+			}
+		}
+
+		if (FolderToUse == nullptr)
+		{
+			FolderToUse = NewObject<UMovieSceneFolder>(MovieScene, NAME_None, RF_Transactional);
+			FolderToUse->SetFolderName(FolderName);
+			MovieScene->GetRootFolders().Add(FolderToUse);
+		}
+
+		return FolderToUse;
+	}
+
 	/** Find or add a folder for the given actor **/
 	static void FindOrAddFolder(UMovieScene* MovieScene, TWeakObjectPtr<AActor> Actor, FGuid Guid)
 	{
@@ -222,53 +272,20 @@ protected:
 			FolderName = TEXT("Misc");
 		}
 
-		// look for a folder to put us in
-		UMovieSceneFolder* FolderToUse = nullptr;
-		for (UMovieSceneFolder* Folder : MovieScene->GetRootFolders())
-		{
-			if (Folder->GetFolderName() == FolderName)
-			{
-				FolderToUse = Folder;
-				break;
-			}
-		}
-
-		if (FolderToUse == nullptr)
-		{
-			FolderToUse = NewObject<UMovieSceneFolder>(MovieScene, NAME_None, RF_Transactional);
-			FolderToUse->SetFolderName(FolderName);
-			MovieScene->GetRootFolders().Add(FolderToUse);
-		}
-
+		UMovieSceneFolder* FolderToUse = FindOrAddFolder(MovieScene, FolderName);
 		FolderToUse->AddChildObjectBinding(Guid);
 	}
 
 	/** Add master track to a folder **/
 	static void AddMasterTrackToFolder(UMovieScene* MovieScene, UMovieSceneTrack* MovieSceneTrack, FName FolderName)
 	{
-		UMovieSceneFolder* FolderToUse = nullptr;
-		for (UMovieSceneFolder* Folder : MovieScene->GetRootFolders())
-		{
-			if (Folder->GetFolderName() == FolderName)
-			{
-				FolderToUse = Folder;
-				break;
-			}
-		}
-
-		if (FolderToUse == nullptr)
-		{
-			FolderToUse = NewObject<UMovieSceneFolder>(MovieScene, NAME_None, RF_Transactional);
-			FolderToUse->SetFolderName(FolderName);
-			MovieScene->GetRootFolders().Add(FolderToUse);
-		}
-
+		UMovieSceneFolder* FolderToUse = FindOrAddFolder(MovieScene, FolderName);
 		FolderToUse->AddChildMasterTrack(MovieSceneTrack);
 	}
 
 	/** Add property to possessable node **/
 	template <typename T>
-	static T* AddPropertyTrack(FName InPropertyName, AActor* InActor, const FGuid& PossessableGuid, UMovieSceneSequence* NewSequence, UMovieScene* NewMovieScene, int32& NumWarnings)
+	static T* AddPropertyTrack(FName InPropertyName, AActor* InActor, const FGuid& PossessableGuid, IMovieScenePlayer& Player, UMovieSceneSequence* NewSequence, UMovieScene* NewMovieScene, int32& NumWarnings)
 	{
 		T* PropertyTrack = nullptr;
 
@@ -281,7 +298,7 @@ protected:
 		if (PropObject && Property)
 		{
 			// If the property object that owns this property isn't already bound, add a binding to the property object
-			ObjectGuid = NewSequence->FindPossessableObjectId(*PropObject);
+			ObjectGuid = Player.FindObjectId(*PropObject, MovieSceneSequenceID::Root);
 			if (!ObjectGuid.IsValid())
 			{
 				UObject* BindingContext = InActor->GetWorld();
@@ -313,7 +330,7 @@ protected:
 				PropertyArray.Insert(Outer, 0);
 				Outer = Outer->GetOuter();
 			}
-			
+
 			FString PropertyPath;
 			for (auto PropertyIt : PropertyArray)
 			{
@@ -342,7 +359,7 @@ protected:
 	}
 
 	/** Convert an interp group */
-	static void ConvertInterpGroup(UInterpGroup* Group, AActor* GroupActor, UMovieScene* NewMovieScene, UMovieSceneSequence* NewSequence, int32& NumWarnings)
+	void ConvertInterpGroup(UInterpGroup* Group, AActor* GroupActor, IMovieScenePlayer& Player, UMovieSceneSequence* NewSequence, UMovieScene* NewMovieScene, int32& NumWarnings)
 	{
 		FGuid PossessableGuid;
 
@@ -365,7 +382,11 @@ protected:
 			}
 
 			// Handle each track class
-			if (Track->IsA(UInterpTrackMove::StaticClass()))					
+			if (ExtendedInterpConverters.Find(Track->GetClass()))
+			{
+				ExtendedInterpConverters.Find(Track->GetClass())->Execute(Track, PossessableGuid, NewMovieScene);
+			}
+			else if (Track->IsA(UInterpTrackMove::StaticClass()))
 			{
 				UInterpTrackMove* MatineeMoveTrack = StaticCast<UInterpTrackMove*>(Track);
 
@@ -445,7 +466,7 @@ protected:
 				UInterpTrackBoolProp* MatineeBoolTrack = StaticCast<UInterpTrackBoolProp*>(Track);
 				if (MatineeBoolTrack->GetNumKeyframes() != 0 && GroupActor && PossessableGuid.IsValid())
 				{
-					UMovieSceneBoolTrack* BoolTrack = AddPropertyTrack<UMovieSceneBoolTrack>(MatineeBoolTrack->PropertyName, GroupActor, PossessableGuid, NewSequence, NewMovieScene, NumWarnings);
+					UMovieSceneBoolTrack* BoolTrack = AddPropertyTrack<UMovieSceneBoolTrack>(MatineeBoolTrack->PropertyName, GroupActor, PossessableGuid, Player, NewSequence, NewMovieScene, NumWarnings);
 					if (BoolTrack)
 					{
 						FMatineeImportTools::CopyInterpBoolTrack(MatineeBoolTrack, BoolTrack);
@@ -457,7 +478,7 @@ protected:
 				UInterpTrackFloatProp* MatineeFloatTrack = StaticCast<UInterpTrackFloatProp*>(Track);
 				if (MatineeFloatTrack->GetNumKeyframes() != 0 && GroupActor && PossessableGuid.IsValid())
 				{
-					UMovieSceneFloatTrack* FloatTrack = AddPropertyTrack<UMovieSceneFloatTrack>(MatineeFloatTrack->PropertyName, GroupActor, PossessableGuid, NewSequence, NewMovieScene, NumWarnings);
+					UMovieSceneFloatTrack* FloatTrack = AddPropertyTrack<UMovieSceneFloatTrack>(MatineeFloatTrack->PropertyName, GroupActor, PossessableGuid, Player, NewSequence, NewMovieScene, NumWarnings);
 					if (FloatTrack)
 					{
 						FMatineeImportTools::CopyInterpFloatTrack(MatineeFloatTrack, FloatTrack);
@@ -469,7 +490,7 @@ protected:
 				UInterpTrackColorProp* MatineeColorTrack = StaticCast<UInterpTrackColorProp*>(Track);
 				if (MatineeColorTrack->GetNumKeyframes() != 0 && GroupActor && PossessableGuid.IsValid())
 				{
-					UMovieSceneColorTrack* ColorTrack = AddPropertyTrack<UMovieSceneColorTrack>(MatineeColorTrack->PropertyName, GroupActor, PossessableGuid, NewSequence, NewMovieScene, NumWarnings);
+					UMovieSceneColorTrack* ColorTrack = AddPropertyTrack<UMovieSceneColorTrack>(MatineeColorTrack->PropertyName, GroupActor, PossessableGuid, Player, NewSequence, NewMovieScene, NumWarnings);
 					if (ColorTrack)
 					{
 						FMatineeImportTools::CopyInterpColorTrack(MatineeColorTrack, ColorTrack);
@@ -481,7 +502,7 @@ protected:
 				UInterpTrackLinearColorProp* MatineeLinearColorTrack = StaticCast<UInterpTrackLinearColorProp*>(Track);
 				if (MatineeLinearColorTrack->GetNumKeyframes() != 0 && GroupActor && PossessableGuid.IsValid())
 				{
-					UMovieSceneColorTrack* ColorTrack = AddPropertyTrack<UMovieSceneColorTrack>(MatineeLinearColorTrack->PropertyName, GroupActor, PossessableGuid, NewSequence, NewMovieScene, NumWarnings);
+					UMovieSceneColorTrack* ColorTrack = AddPropertyTrack<UMovieSceneColorTrack>(MatineeLinearColorTrack->PropertyName, GroupActor, PossessableGuid, Player, NewSequence, NewMovieScene, NumWarnings);
 					if (ColorTrack)
 					{
 						FMatineeImportTools::CopyInterpLinearColorTrack(MatineeLinearColorTrack, ColorTrack);
@@ -523,7 +544,7 @@ protected:
 	}
 
 	/** Convert a single matinee to a level sequence asset */
-	static TWeakObjectPtr<ALevelSequenceActor> ConvertSingleMatineeToLevelSequence(TWeakObjectPtr<AActor> ActorToConvert, int32& NumWarnings)
+	TWeakObjectPtr<ALevelSequenceActor> ConvertSingleMatineeToLevelSequence(TWeakObjectPtr<AActor> ActorToConvert, int32& NumWarnings)
 	{
 		UObject* AssetOuter = ActorToConvert->GetOuter();
 		UPackage* AssetPackage = AssetOuter->GetOutermost();
@@ -568,6 +589,27 @@ protected:
 
 		ALevelSequenceActor* NewActor = CastChecked<ALevelSequenceActor>(GEditor->UseActorFactory(ActorFactory, FAssetData(NewAsset), &FTransform::Identity));
 
+		struct FTemporaryPlayer : IMovieScenePlayer
+		{
+			FTemporaryPlayer(UMovieSceneSequence& InSequence, UObject* InContext)
+				: Context(InContext)
+			{
+				RootInstance.Initialize(InSequence, *this);
+			}
+
+			virtual FMovieSceneRootEvaluationTemplateInstance& GetEvaluationTemplate() { return RootInstance; }
+			virtual void UpdateCameraCut(UObject* CameraObject, UObject* UnlockIfCameraObject = nullptr, bool bJumpCut = false) {}
+			virtual void SetViewportSettings(const TMap<FViewportClient*, EMovieSceneViewportParams>& ViewportParamsMap) {}
+			virtual void GetViewportSettings(TMap<FViewportClient*, EMovieSceneViewportParams>& ViewportParamsMap) const {}
+			virtual EMovieScenePlayerStatus::Type GetPlaybackStatus() const { return EMovieScenePlayerStatus::Stopped; }
+			virtual void SetPlaybackStatus(EMovieScenePlayerStatus::Type InPlaybackStatus) {}
+			virtual UObject* GetPlaybackContext() const { return Context; }
+
+			FMovieSceneRootEvaluationTemplateInstance RootInstance;
+			UObject* Context;
+
+		} TemporaryPlayer(*NewSequence, NewActor->GetWorld());
+
 		// Walk through all the interp group data and create corresponding tracks on the new level sequence asset
 		if (ActorToConvert->IsA(AMatineeActor::StaticClass()))
 		{
@@ -583,7 +625,7 @@ protected:
 				UInterpGroupInst* GrInst = MatineeActor->GroupInst[i];
 				UInterpGroup* Group = GrInst->Group;
 				AActor* GroupActor = GrInst->GetGroupActor();
-				ConvertInterpGroup(Group, GroupActor, NewMovieScene, NewSequence, NumWarnings);
+				ConvertInterpGroup(Group, GroupActor, TemporaryPlayer, NewSequence, NewMovieScene, NumWarnings);
 			}
 
 			// Director group - convert this after the regular groups to ensure that the camera cut bindings are there
@@ -594,7 +636,7 @@ protected:
 				if (MatineeDirectorTrack && MatineeDirectorTrack->GetNumKeyframes() != 0)
 				{
 					UMovieSceneCameraCutTrack* CameraCutTrack = Cast<UMovieSceneCameraCutTrack>(NewMovieScene->AddMasterTrack<UMovieSceneCameraCutTrack>());
-					FMatineeImportTools::CopyInterpDirectorTrack(MatineeDirectorTrack, CameraCutTrack, MatineeActor, NewSequence);
+					FMatineeImportTools::CopyInterpDirectorTrack(MatineeDirectorTrack, CameraCutTrack, MatineeActor, TemporaryPlayer);
 				}
 
 				UInterpTrackFade* MatineeFadeTrack = DirGroup->GetFadeTrack();
@@ -637,6 +679,8 @@ private:
 
 	FDelegateHandle LevelEditorExtenderDelegateHandle;
 
+	// IMatineeToLevelSequenceModule interface
+	TMap<TSubclassOf<UInterpTrack>, FOnConvertMatineeTrack > ExtendedInterpConverters;
 };
 
 IMPLEMENT_MODULE(FMatineeToLevelSequenceModule, MatineeToLevelSequence);
diff --git a/Engine/Plugins/MovieScene/MatineeToLevelSequence/Source/MatineeToLevelSequence/Public/MatineeToLevelSequenceModule.h b/Engine/Plugins/MovieScene/MatineeToLevelSequence/Source/MatineeToLevelSequence/Public/MatineeToLevelSequenceModule.h
new file mode 100644
--- /dev/null
+++ b/Engine/Plugins/MovieScene/MatineeToLevelSequence/Source/MatineeToLevelSequence/Public/MatineeToLevelSequenceModule.h
@@ -0,0 +1,19 @@
+// Copyright 1998-2016 Epic Games, Inc. All Rights Reserved.
+
+#pragma once
+
+#define MATINEE_TO_LEVEL_SEQUENCE_MODULE_NAME "MatineeToLevelSequence"
+
+//////////////////////////////////////////////////////////////////////////
+// IMatineeToLevelSequenceModule
+
+class IMatineeToLevelSequenceModule : public IModuleInterface
+{
+public:
+	
+ 	DECLARE_DELEGATE_ThreeParams(FOnConvertMatineeTrack, UInterpTrack*, FGuid, UMovieScene*);
+ 
+ 	FDelegateHandle RegisterTrackConverterForMatineeClass(TSubclassOf<UInterpTrack> InterpTrackClass, FOnConvertMatineeTrack);
+ 	void UnregisterTrackConverterForMatineeClass(FDelegateHandle RemoveDelegate);
+};
+
diff --git a/Engine/Plugins/Slate/SlateRemote/SlateRemote.uplugin b/Engine/Plugins/Slate/SlateRemote/SlateRemote.uplugin
--- a/Engine/Plugins/Slate/SlateRemote/SlateRemote.uplugin
+++ b/Engine/Plugins/Slate/SlateRemote/SlateRemote.uplugin
@@ -10,7 +10,7 @@
 	"DocsURL" : "",
 	"MarketplaceURL" : "",
 	"SupportURL" : "",
-	"EnabledByDefault" : true,
+	"EnabledByDefault" : false,
 	"CanContainContent" : false,
 	"IsBetaVersion" : false,
 	"Installed" : false,
diff --git a/Engine/Plugins/Slate/SlateRemote/Source/SlateRemote/Private/SlateRemoteClasses.cpp b/Engine/Plugins/Slate/SlateRemote/Source/SlateRemote/Private/SlateRemoteClasses.cpp
--- a/Engine/Plugins/Slate/SlateRemote/Source/SlateRemote/Private/SlateRemoteClasses.cpp
+++ b/Engine/Plugins/Slate/SlateRemote/Source/SlateRemote/Private/SlateRemoteClasses.cpp
@@ -5,7 +5,7 @@
 
 USlateRemoteSettings::USlateRemoteSettings( const FObjectInitializer& ObjectInitializer )
 	: Super(ObjectInitializer)
-	, EnableRemoteServer(false)
+	, EnableRemoteServer(true)
 	, EditorServerEndpoint(SLATE_REMOTE_SERVER_DEFAULT_EDITOR_ENDPOINT.ToString())
 	, GameServerEndpoint(SLATE_REMOTE_SERVER_DEFAULT_GAME_ENDPOINT.ToString())
 { }
diff --git a/Engine/Shaders/MediaShaders.usf b/Engine/Shaders/MediaShaders.usf
--- a/Engine/Shaders/MediaShaders.usf
+++ b/Engine/Shaders/MediaShaders.usf
@@ -88,21 +88,38 @@ void AYUVConvertPS(
 }
 
 
+// shader to convret a Windows Bitmap texture
+void BMPConvertPS(
+	float4 InPosition : SV_POSITION,
+	float2 InUV : TEXCOORD0,
+	out float4 OutColor : SV_Target0)
+{
+	InUV.y = 1.0f - InUV.y;
+	float2 ScaledUV = InUV * BMPConvertUB.UVScale;
+	OutColor = BMPConvertUB.BMPTexture.Sample(BMPConvertUB.BMPSampler, ScaledUV);
+}
+
+
 // shader to convert a NV12 texture to RGBA
 void NV12ConvertPS(
 	float4 InPosition : SV_POSITION,
 	float2 InUV : TEXCOORD0,
 	out float4 OutColor : SV_Target0)
 {
+	float2 ScaledUV = InUV * NV12ConvertUB.UVScale;
+
 	float3 YUV;
-	YUV.x = NV12ConvertUB.NV12Texture.Sample(NV12ConvertUB.NV12Sampler, float2(InUV.x, InUV.y / 3.0f * 2.0f)).x;
+	YUV.x = NV12ConvertUB.NV12Texture.Sample(NV12ConvertUB.NV12SamplerB, ScaledUV).x;
+
+	ScaledUV.y *= 0.5f;
+	ScaledUV.y += 2.0f / 3.0f;
 
-	uint UVOffs = uint(InUV.x * NV12ConvertUB.Width / 2.0f) * 2;
-	float UPos = (float(UVOffs) + 0.5f) / NV12ConvertUB.Width;
-	float VPos = (float(UVOffs) + 1.5f) / NV12ConvertUB.Width;
+	uint UVOffs = uint(InUV.x * NV12ConvertUB.OutputWidth / 2.0f) * 2;
+	float UPos = (float(UVOffs * NV12ConvertUB.UVScale.x) + 0.5f) / NV12ConvertUB.OutputWidth;
+	float VPos = (float(UVOffs * NV12ConvertUB.UVScale.x) + 1.5f) / NV12ConvertUB.OutputWidth;
 
-	YUV.y = NV12ConvertUB.NV12Texture.Sample(NV12ConvertUB.NV12Sampler, float2(UPos, (InUV.y + 2.0f) / 3.0f)).x;
-	YUV.z = NV12ConvertUB.NV12Texture.Sample(NV12ConvertUB.NV12Sampler, float2(VPos, (InUV.y + 2.0f) / 3.0f)).x;
+	YUV.y = NV12ConvertUB.NV12Texture.Sample(NV12ConvertUB.NV12SamplerP, float2(UPos, ScaledUV.y)).x;
+	YUV.z = NV12ConvertUB.NV12Texture.Sample(NV12ConvertUB.NV12SamplerP, float2(VPos, ScaledUV.y)).x;
 
 	OutColor.xyz = YuvToRgb(YUV, YUV_TO_RGB_MSDN);
 	OutColor.w = 1.0f;
@@ -115,21 +132,38 @@ void NV21ConvertPS(
 	float2 InUV : TEXCOORD0,
 	out float4 OutColor : SV_Target0)
 {
+	float2 ScaledUV = InUV * NV21ConvertUB.UVScale;
+
 	float3 YUV;
-	YUV.x = NV12ConvertUB.NV12Texture.Sample(NV12ConvertUB.NV12Sampler, float2(InUV.x, InUV.y / 3.0f * 2.0f)).x;
+	YUV.x = NV21ConvertUB.NV21Texture.Sample(NV21ConvertUB.NV21SamplerB, ScaledUV).x;
+
+	ScaledUV.y *= 0.5f;
+	ScaledUV.y += 2.0f / 3.0f;
 
-	uint VUPos = uint(InUV.x * NV12ConvertUB.Width / 2.0f) * 2;
-	float UPos = (float(VUPos) + 1.5f) / NV12ConvertUB.Width;
-	float VPos = (float(VUPos) + 0.5f) / NV12ConvertUB.Width;
+	uint UVOffs = uint(ScaledUV.x * NV21ConvertUB.OutputWidth / 2.0f) * 2;
+	float UPos = (float(UVOffs * NV21ConvertUB.UVScale.x) + 1.5f) / NV21ConvertUB.OutputWidth;
+	float VPos = (float(UVOffs * NV21ConvertUB.UVScale.x) + 0.5f) / NV21ConvertUB.OutputWidth;
 
-	YUV.y = NV12ConvertUB.NV12Texture.Sample(NV12ConvertUB.NV12Sampler, float2(UPos, (InUV.y + 2.0f) / 3.0f)).x;
-	YUV.z = NV12ConvertUB.NV12Texture.Sample(NV12ConvertUB.NV12Sampler, float2(VPos, (InUV.y + 2.0f) / 3.0f)).x;
+	YUV.y = NV21ConvertUB.NV21Texture.Sample(NV21ConvertUB.NV21SamplerP, float2(UPos, ScaledUV.y)).x;
+	YUV.z = NV21ConvertUB.NV21Texture.Sample(NV21ConvertUB.NV21SamplerP, float2(VPos, ScaledUV.y)).x;
 
 	OutColor.xyz = YuvToRgb(YUV, YUV_TO_RGB_MSDN);
 	OutColor.w = 1.0f;
 }
 
 
+// shader to resize an RGB texture
+void RGBConvertPS(
+	float4 InPosition : SV_POSITION,
+	float2 InUV : TEXCOORD0,
+	out float4 OutColor : SV_Target0)
+{
+	float2 ScaledUV = InUV * RGBConvertUB.UVScale;
+	float4 RGB = RGBConvertUB.RGBTexture.Sample(RGBConvertUB.RGBSampler, ScaledUV);
+	OutColor.xyz = sRGBToLinear(RGB);
+}
+
+
 // shader to convert a luma (Y) and chroma (CbCr) texture to RGBA
 void YCbCrConvertPS(
 	float4 InPosition : SV_POSITION,
@@ -187,14 +221,14 @@ void YUY2ConvertPS(
 	float2 InUV : TEXCOORD0,
 	out float4 OutColor : SV_Target0)
 {
-	float4 YUY2B = YUY2ConvertUB.YUY2Texture.Sample(YUY2ConvertUB.YUY2SamplerB, float2(InUV.x, InUV.y));
-	float4 YUY2P = YUY2ConvertUB.YUY2Texture.Sample(YUY2ConvertUB.YUY2SamplerP, float2(InUV.x, InUV.y));
+	float2 ScaledUV = InUV * YUY2ConvertUB.UVScale;
 
-	float Pos = 2.0f * InUV.x * YUY2ConvertUB.Width;
 	float3 YUV;
+	YUV.yz = YUY2ConvertUB.YUY2Texture.Sample(YUY2ConvertUB.YUY2SamplerB, ScaledUV).yw;
 
+	float4 YUY2P = YUY2ConvertUB.YUY2Texture.Sample(YUY2ConvertUB.YUY2SamplerP, ScaledUV);
+	float Pos = InUV.x * YUY2ConvertUB.OutputWidth;
 	YUV.x = (floor(Pos % 2.0f) == 0) ? YUY2P.z : YUY2P.x;
-	YUV.yz = YUY2B.yw;
 
 	OutColor.xyz = YuvToRgb(YUV, YUV_TO_RGB_MSDN);
 	OutColor.w = 1.0f;
diff --git a/Engine/Source/Developer/Android/AndroidDeviceDetection/AndroidDeviceDetection.Build.cs b/Engine/Source/Developer/Android/AndroidDeviceDetection/AndroidDeviceDetection.Build.cs
--- a/Engine/Source/Developer/Android/AndroidDeviceDetection/AndroidDeviceDetection.Build.cs
+++ b/Engine/Source/Developer/Android/AndroidDeviceDetection/AndroidDeviceDetection.Build.cs
@@ -12,7 +12,13 @@ public class AndroidDeviceDetection : ModuleRules
 			new string[]
 			{
 				"Core",
-                "Messaging",
+			}
+		);
+
+		PrivateIncludePathModuleNames.AddRange(
+			new string[]
+			{
+				"TcpMessaging",
 			}
 		);
 
diff --git a/Engine/Source/Developer/Android/AndroidDeviceDetection/Private/AndroidDeviceDetectionModule.cpp b/Engine/Source/Developer/Android/AndroidDeviceDetection/Private/AndroidDeviceDetectionModule.cpp
--- a/Engine/Source/Developer/Android/AndroidDeviceDetection/Private/AndroidDeviceDetectionModule.cpp
+++ b/Engine/Source/Developer/Android/AndroidDeviceDetection/Private/AndroidDeviceDetectionModule.cpp
@@ -6,7 +6,7 @@
 
 #include "AndroidDeviceDetectionPrivatePCH.h"
 #include "ModuleManager.h"
-#include "IConnectionBasedMessagingModule.h"
+#include "ITcpMessagingModule.h"
 
 #define LOCTEXT_NAMESPACE "FAndroidDeviceDetectionModule" 
 
@@ -23,7 +23,7 @@ public:
 		HasADBPath(false),
 		ForceCheck(false)
 	{
-		TcpMessagingModule = FModuleManager::LoadModulePtr<IConnectionBasedMessagingModule>("TcpMessaging");
+		TcpMessagingModule = FModuleManager::LoadModulePtr<ITcpMessagingModule>("TcpMessaging");
 	}
 
 public:
@@ -342,7 +342,7 @@ private:
 	bool HasADBPath;
 	bool ForceCheck;
 
-	IConnectionBasedMessagingModule* TcpMessagingModule;
+	ITcpMessagingModule* TcpMessagingModule;
 };
 
 class FAndroidDeviceDetection : public IAndroidDeviceDetection
diff --git a/Engine/Source/Developer/AutomationWindow/Private/SAutomationWindow.cpp b/Engine/Source/Developer/AutomationWindow/Private/SAutomationWindow.cpp
--- a/Engine/Source/Developer/AutomationWindow/Private/SAutomationWindow.cpp
+++ b/Engine/Source/Developer/AutomationWindow/Private/SAutomationWindow.cpp
@@ -104,7 +104,7 @@ SAutomationWindow::~SAutomationWindow()
 }
 
 BEGIN_SLATE_FUNCTION_BUILD_OPTIMIZATION
-void SAutomationWindow::Construct( const FArguments& InArgs, const IAutomationControllerManagerRef& InAutomationController, const ISessionManagerRef& InSessionManager )
+void SAutomationWindow::Construct( const FArguments& InArgs, const IAutomationControllerManagerRef& InAutomationController, const TSharedRef<ISessionManager>& InSessionManager )
 {
 	FAutomationWindowCommands::Register();
 	CreateCommands();
@@ -1959,7 +1959,7 @@ void SAutomationWindow::HandleSessionManagerInstanceChanged()
 
 void SAutomationWindow::UpdateTestListBackgroundStyle()
 {
-	TArray<ISessionInstanceInfoPtr> OutInstances;
+	TArray<TSharedPtr<ISessionInstanceInfo>> OutInstances;
 
 	if( ActiveSession.IsValid() )
 	{
@@ -2208,7 +2208,7 @@ EVisibility SAutomationWindow::HandleSelectSessionOverlayVisibility( ) const
 }
 
 
-void SAutomationWindow::HandleSessionManagerCanSelectSession( const ISessionInfoPtr& Session, bool& CanSelect )
+void SAutomationWindow::HandleSessionManagerCanSelectSession( const TSharedPtr<ISessionInfo>& Session, bool& CanSelect )
 {
 	if (ActiveSession.IsValid() && AutomationController->CheckTestResultsAvailable())
 	{
@@ -2218,7 +2218,7 @@ void SAutomationWindow::HandleSessionManagerCanSelectSession( const ISessionInfo
 }
 
 
-void SAutomationWindow::HandleSessionManagerSelectionChanged( const ISessionInfoPtr& SelectedSession )
+void SAutomationWindow::HandleSessionManagerSelectionChanged( const TSharedPtr<ISessionInfo>& SelectedSession )
 {
 	FindWorkers();
 }
diff --git a/Engine/Source/Developer/AutomationWindow/Private/SAutomationWindow.h b/Engine/Source/Developer/AutomationWindow/Private/SAutomationWindow.h
--- a/Engine/Source/Developer/AutomationWindow/Private/SAutomationWindow.h
+++ b/Engine/Source/Developer/AutomationWindow/Private/SAutomationWindow.h
@@ -76,7 +76,7 @@ public:
 	/**
 	 * Constructs the widget.
 	 */
-	void Construct(const FArguments& InArgs, const IAutomationControllerManagerRef& AutomationController, const ISessionManagerRef& InSessionManager);
+	void Construct(const FArguments& InArgs, const IAutomationControllerManagerRef& AutomationController, const TSharedRef<ISessionManager>& InSessionManager);
 
 	/**
 	* Check tests aren't running.
@@ -373,14 +373,14 @@ private:
 	EVisibility HandleSelectSessionOverlayVisibility( ) const;
 
 	/** Callback for determining whether a session can be selected in the session manager. */
-	void HandleSessionManagerCanSelectSession( const ISessionInfoPtr& Session, bool& CanSelect );
+	void HandleSessionManagerCanSelectSession( const TSharedPtr<ISessionInfo>& Session, bool& CanSelect );
 
 	/**
 	 * Session selection has changed in the session manager
 	 *
 	 * @param SelectedSession The session that was selected.
 	 */
-	void HandleSessionManagerSelectionChanged( const ISessionInfoPtr& SelectedSession );
+	void HandleSessionManagerSelectionChanged( const TSharedPtr<ISessionInfo>& SelectedSession );
 
 	/** Called when the session manager updates an instances. */
 	void HandleSessionManagerInstanceChanged();
@@ -526,7 +526,7 @@ private:
 	TSharedPtr<FUICommandList> AutomationWindowActions;
 
 	/** Holds a pointer to the active session. */
-	ISessionInfoPtr ActiveSession;
+	TSharedPtr<ISessionInfo> ActiveSession;
 
 	/** Holds the AutomationController. */
 	IAutomationControllerManagerPtr AutomationController;
@@ -571,7 +571,7 @@ private:
 	TSharedPtr< AutomationFilterCollection > AutomationFilters;
 
 	/** Holds the session manager. */
-	ISessionManagerPtr SessionManager;
+	TSharedPtr<ISessionManager> SessionManager;
 
 	/**
 	 * Holds the automation controller module state.
diff --git a/Engine/Source/Developer/LauncherServices/Private/LauncherServicesModule.cpp b/Engine/Source/Developer/LauncherServices/Private/LauncherServicesModule.cpp
--- a/Engine/Source/Developer/LauncherServices/Private/LauncherServicesModule.cpp
+++ b/Engine/Source/Developer/LauncherServices/Private/LauncherServicesModule.cpp
@@ -11,30 +11,30 @@ class FLauncherServicesModule
 {
 public:
 
-	// Begin ILauncherServicesModule interface
+	//~ ILauncherServicesModule interface
 
-	virtual ILauncherDeviceGroupRef CreateDeviceGroup( ) override
+	virtual ILauncherDeviceGroupRef CreateDeviceGroup() override
 	{
 		return MakeShareable(new FLauncherDeviceGroup());
 	}
 
-	virtual ILauncherDeviceGroupRef CreateDeviceGroup(const FGuid& Guid, const FString& Name ) override
+	virtual ILauncherDeviceGroupRef CreateDeviceGroup(const FGuid& Guid, const FString& Name) override
 	{
 		return MakeShareable(new FLauncherDeviceGroup(Guid, Name));
 	}
 
-	virtual ILauncherRef CreateLauncher( ) override
+	virtual ILauncherRef CreateLauncher() override
 	{
 		return MakeShareable(new FLauncher());
 	}
 
-	virtual ILauncherProfileRef CreateProfile( const FString& ProfileName ) override
+	virtual ILauncherProfileRef CreateProfile(const FString& ProfileName) override
 	{
 		ILauncherProfileManagerRef ProfileManager = GetProfileManager();
 		return MakeShareable(new FLauncherProfile(ProfileManager, FGuid(), ProfileName));
 	}
 
-	virtual ILauncherProfileManagerRef GetProfileManager( ) override
+	virtual ILauncherProfileManagerRef GetProfileManager() override
 	{
 		if (!ProfileManagerSingleton.IsValid())
 		{
@@ -48,7 +48,7 @@ public:
 	}
 	
 	DECLARE_DERIVED_EVENT(FLauncherServicesModule, ILauncherServicesModule::FLauncherServicesSDKNotInstalled, FLauncherServicesSDKNotInstalled);
-	virtual FLauncherServicesSDKNotInstalled& OnLauncherServicesSDKNotInstalled( ) override
+	virtual FLauncherServicesSDKNotInstalled& OnLauncherServicesSDKNotInstalled() override
 	{
 		return LauncherServicesSDKNotInstalled;
 	}
@@ -57,14 +57,12 @@ public:
 		return LauncherServicesSDKNotInstalled.Broadcast(PlatformName, DocLink);
 	}
 
-	// End ILauncherServicesModule interface
-
 private:
 	
-	/// Event to be called when the editor tried to use a platform, but it wasn't installed
+	/** Event to be called when the editor tried to use a platform, but it wasn't installed. */
 	FLauncherServicesSDKNotInstalled LauncherServicesSDKNotInstalled;
 
-	// Holds the launcher profile manager singleton
+	/** The launcher profile manager singleton. */
 	static ILauncherProfileManagerPtr ProfileManagerSingleton;
 };
 
diff --git a/Engine/Source/Developer/LauncherServices/Private/LauncherServicesPrivatePCH.h b/Engine/Source/Developer/LauncherServices/Private/LauncherServicesPrivatePCH.h
--- a/Engine/Source/Developer/LauncherServices/Private/LauncherServicesPrivatePCH.h
+++ b/Engine/Source/Developer/LauncherServices/Private/LauncherServicesPrivatePCH.h
@@ -10,7 +10,7 @@
 
 #include "Messaging.h"
 #include "ModuleManager.h"
-#include "SessionMessages.h"
+#include "Runtime/SessionMessages/Public/SessionServiceMessages.h"
 #include "TargetDeviceServices.h"
 #include "TargetPlatform.h"
 #include "Ticker.h"
diff --git a/Engine/Source/Developer/Profiler/Private/ProfilerCommands.cpp b/Engine/Source/Developer/Profiler/Private/ProfilerCommands.cpp
--- a/Engine/Source/Developer/Profiler/Private/ProfilerCommands.cpp
+++ b/Engine/Source/Developer/Profiler/Private/ProfilerCommands.cpp
@@ -1,23 +1,29 @@
 // Copyright 1998-2016 Epic Games, Inc. All Rights Reserved.
 
 #include "ProfilerPrivatePCH.h"
-#include "DesktopPlatformModule.h"
-#include "FileManagerGeneric.h"
+#include "ProfilerCommands.h"
+#include "ProfilerManager.h"
+#include "ProfilerSession.h"
+#include "SMultiDumpBrowser.h"
+#include "SProfilerWindow.h"
+
 
 #define LOCTEXT_NAMESPACE "FProfilerCommands"
 
+
 /*-----------------------------------------------------------------------------
 	FProfilerCommands
 -----------------------------------------------------------------------------*/
 
-FProfilerCommands::FProfilerCommands() : TCommands<FProfilerCommands>
-(
-	TEXT( "ProfilerCommand" ), // Context name for fast lookup
-	NSLOCTEXT("Contexts", "ProfilerCommand", "Profiler Command"), // Localized context name for displaying
-	NAME_None, // Parent
-	FEditorStyle::GetStyleSetName() // Icon Style Set
-)
-{}
+FProfilerCommands::FProfilerCommands()
+	: TCommands<FProfilerCommands>(
+		TEXT( "ProfilerCommand" ), // Context name for fast lookup
+		NSLOCTEXT("Contexts", "ProfilerCommand", "Profiler Command"), // Localized context name for displaying
+		NAME_None, // Parent
+		FEditorStyle::GetStyleSetName() // Icon Style Set
+	)
+{ }
+
 
 /** UI_COMMAND takes long for the compile to optimize */
 PRAGMA_DISABLE_OPTIMIZATION
@@ -63,23 +69,25 @@ PRAGMA_ENABLE_OPTIMIZATION
 
 //static_assert(sizeof(FProfilerActionManager) == 0, "Cannot contain any variables at this moment.");
 
+
 /*-----------------------------------------------------------------------------
 	FProfilerMenuBuilder
 -----------------------------------------------------------------------------*/
 
-void FProfilerMenuBuilder::AddMenuEntry( FMenuBuilder& MenuBuilder, const TSharedPtr< FUICommandInfo >& FUICommandInfo, const FUIAction& UIAction )
+void FProfilerMenuBuilder::AddMenuEntry( FMenuBuilder& MenuBuilder, const TSharedPtr< FUICommandInfo >& UICommandInfo, const FUIAction& UIAction )
 {
 	MenuBuilder.AddMenuEntry
 	( 
-		FUICommandInfo->GetLabel(),
-		FUICommandInfo->GetDescription(),
-		FUICommandInfo->GetIcon(),
+		UICommandInfo->GetLabel(),
+		UICommandInfo->GetDescription(),
+		UICommandInfo->GetIcon(),
 		UIAction,
 		NAME_None, 
-		FUICommandInfo->GetUserInterfaceType()
+		UICommandInfo->GetUserInterfaceType()
 	);
 }
 
+
 /*-----------------------------------------------------------------------------
 	ToggleDataPreview
 -----------------------------------------------------------------------------*/
@@ -121,6 +129,7 @@ ECheckBoxState FProfilerActionManager::ToggleDataPreview_GetCheckState( const FG
 	return bDataPreview ? ECheckBoxState::Checked : ECheckBoxState::Unchecked;
 }
 
+
 /*-----------------------------------------------------------------------------
 	ProfilerManager_ToggleLivePreview
 -----------------------------------------------------------------------------*/
@@ -151,6 +160,7 @@ ECheckBoxState FProfilerActionManager::ProfilerManager_ToggleLivePreview_GetChec
 	return This->bLivePreview ? ECheckBoxState::Checked : ECheckBoxState::Unchecked;
 }
 
+
 /*-----------------------------------------------------------------------------
 	ProfilerManager_Load
 -----------------------------------------------------------------------------*/
@@ -214,7 +224,6 @@ void FProfilerActionManager::ProfilerManager_Load_Execute()
 }
 
 
-
 void FProfilerActionManager::ProfilerManager_LoadMultiple_Execute()
 {
 	// @see FStatConstants::StatsFileExtension
@@ -259,6 +268,7 @@ bool FProfilerActionManager::ProfilerManager_Load_CanExecute() const
 	return !(bIsConnectionActive && This->ProfilerType == EProfilerSessionTypes::Live);
 }
 
+
 /*-----------------------------------------------------------------------------
 	ToggleDataCapture
 -----------------------------------------------------------------------------*/
@@ -312,6 +322,7 @@ ECheckBoxState FProfilerActionManager::ToggleDataCapture_GetCheckState( const FG
 	return bDataCapturing ? ECheckBoxState::Checked : ECheckBoxState::Unchecked;
 }
 
+
 /*-----------------------------------------------------------------------------
 		OpenSettings
 -----------------------------------------------------------------------------*/
diff --git a/Engine/Source/Developer/Profiler/Private/ProfilerCommands.h b/Engine/Source/Developer/Profiler/Private/ProfilerCommands.h
--- a/Engine/Source/Developer/Profiler/Private/ProfilerCommands.h
+++ b/Engine/Source/Developer/Profiler/Private/ProfilerCommands.h
@@ -2,23 +2,29 @@
 
 #pragma once
 
-/*-----------------------------------------------------------------------------
-	Declarations
------------------------------------------------------------------------------*/
+#include "Commands.h"
 
-/** Class that holds all profiler commands. */
-class FProfilerCommands : public TCommands<FProfilerCommands>
+
+class FMenuBuilder;
+class FUICommandInfo;
+
+
+/**
+ * Class that holds all profiler commands.
+ */
+class FProfilerCommands
+	: public TCommands<FProfilerCommands>
 {
 public:
+
 	/** Default constructor. */
 	FProfilerCommands();
 	
-	/**
-	 * Initialize commands.
-	 */
+	/** Initialize commands. */
 	virtual void RegisterCommands() override;
 
 public:
+
 	/*-----------------------------------------------------------------------------
 		Global and custom commands.	Need to implement following methods:
 
@@ -82,6 +88,7 @@ public:
 	TSharedPtr< FUICommandInfo > EventGraph_SelectAllFrames;
 };
 
+
 class FProfilerMenuBuilder
 {
 public:
@@ -93,14 +100,15 @@ public:
 	 *		SessionInstance_ToggleCapture			- Global version will toggle capture process for all active session instances
 	 *		SessionInstance_ToggleCapture_OneParam	- Local version will toggle capture process only for the specified session instance
 	 *
-	 * @param MenuBuilder		- the menu to add items to
-	 * @param FUICommandInfo	- a shared pointer to the UI command info
-	 * @param UIAction			- customized version of the UI command info stored in an UI action 
+	 * @param MenuBuilder The menu to add items to
+	 * @param FUICommandInfo A shared pointer to the UI command info
+	 * @param UIAction Customized version of the UI command info stored in an UI action 
 	 *
 	 */
-	static void AddMenuEntry( FMenuBuilder& MenuBuilder, const TSharedPtr< FUICommandInfo >& FUICommandInfo, const FUIAction& UIAction );
+	static void AddMenuEntry( FMenuBuilder& MenuBuilder, const TSharedPtr< FUICommandInfo >& UICommandInfo, const FUIAction& UIAction );
 };
 
+
 /** 
  * Class that provides helper functions for the commands to avoid cluttering profiler manager with many small functions. Can't contain any variables.
  * Directly operates on the profiler manager instance.
@@ -213,4 +221,4 @@ protected:
 
 	/** Reference to the global instance of the profiler manager. */
 	class FProfilerManager* This;
-};
\ No newline at end of file
+};
diff --git a/Engine/Source/Developer/Profiler/Private/ProfilerDataProvider.cpp b/Engine/Source/Developer/Profiler/Private/ProfilerDataProvider.cpp
--- a/Engine/Source/Developer/Profiler/Private/ProfilerDataProvider.cpp
+++ b/Engine/Source/Developer/Profiler/Private/ProfilerDataProvider.cpp
@@ -1,42 +1,46 @@
 // Copyright 1998-2016 Epic Games, Inc. All Rights Reserved.
 
 #include "ProfilerPrivatePCH.h"
+#include "ProfilerCommon.h"
+#include "ProfilerDataProvider.h"
+
 
 /*-----------------------------------------------------------------------------
 	IDataProvider
 -----------------------------------------------------------------------------*/
 
-void IDataProvider::InternalDuplicate( const IDataProviderRef& DataProvider, const uint32 FrameStartIndex, const uint32 FrameEndIndex ) const
+void IDataProvider::InternalDuplicate(const TSharedRef<IDataProvider>& DataProvider, const uint32 FrameStartIndex, const uint32 FrameEndIndex) const
 {
-	for( uint32 FrameIdx = FrameStartIndex; FrameIdx < FrameEndIndex; FrameIdx++ )
+	for(uint32 FrameIdx = FrameStartIndex; FrameIdx < FrameEndIndex; FrameIdx++)
 	{
-		const uint32 SampleStartIndex = GetSamplesIndicesForFrame( FrameIdx ).X;
-		const uint32 SampleEndIndex = GetSamplesIndicesForFrame( FrameIdx ).Y;
+		const uint32 SampleStartIndex = GetSamplesIndicesForFrame(FrameIdx).X;
+		const uint32 SampleEndIndex = GetSamplesIndicesForFrame(FrameIdx).Y;
 
-		for( uint32 SampleIndex = SampleStartIndex; SampleIndex < SampleEndIndex; SampleIndex++ )
+		for(uint32 SampleIndex = SampleStartIndex; SampleIndex < SampleEndIndex; SampleIndex++)
 		{
-			FProfilerSample ProfilerSample( GetCollection()[ SampleIndex ] );
+			FProfilerSample ProfilerSample(GetCollection()[ SampleIndex ]);
 
-			for( int32 ChildIdx = 0; ChildIdx < ProfilerSample._ChildrenIndices.Num(); ChildIdx++ )
+			for(int32 ChildIdx = 0; ChildIdx < ProfilerSample._ChildrenIndices.Num(); ChildIdx++)
 			{
 				ProfilerSample._ChildrenIndices[ChildIdx] -= SampleStartIndex;
 			}
 
-			DataProvider->AddDuplicatedSample( ProfilerSample );
+			DataProvider->AddDuplicatedSample(ProfilerSample);
 		}
-		DataProvider->AdvanceFrame( GetFrameTimeMS( FrameIdx ) );
+		DataProvider->AdvanceFrame(GetFrameTimeMS(FrameIdx));
 	}
 }
 
-const uint32 IDataProvider::AdvanceFrame( const float DeltaTimeMS )
+
+const uint32 IDataProvider::AdvanceFrame(const float DeltaTimeMS)
 {
-	if( !bHasAddedSecondStartMarker )
+	if(!bHasAddedSecondStartMarker)
 	{
 		bHasAddedSecondStartMarker = true;
 
 		// Add the default values.
-		FrameCounters.Add( LastSecondFrameCounter );
-		AccumulatedFrameCounters.Add( GetNumFrames() );
+		FrameCounters.Add(LastSecondFrameCounter);
+		AccumulatedFrameCounters.Add(GetNumFrames());
 	}
 
 	ElapsedTimeMS += DeltaTimeMS;
@@ -45,23 +49,23 @@ const uint32 IDataProvider::AdvanceFrame( const float DeltaTimeMS )
 	LastSecondFrameCounter ++;
 
 	const uint32 SampleEndIndex = GetNumSamples();
-	FrameIndices.Add( FIntPoint( FrameIndex, SampleEndIndex ) );
+	FrameIndices.Add(FIntPoint(FrameIndex, SampleEndIndex));
 
-	FrameTimes.Add( DeltaTimeMS );
-	ElapsedFrameTimes.Add( ElapsedTimeMS );
+	FrameTimes.Add(DeltaTimeMS);
+	ElapsedFrameTimes.Add(ElapsedTimeMS);
 
 	// Update the values.
 	FrameCounters.Last() = LastSecondFrameCounter;
 	AccumulatedFrameCounters.Last() = GetNumFrames();
 
 	int NumLongFrames = 0;
-	while( LastSecondFrameTimeMS > 1000.0f )
+	while(LastSecondFrameTimeMS > 1000.0f)
 	{
-		if( NumLongFrames > 0 )
+		if(NumLongFrames > 0)
 		{
 			// Add the default values.
-			FrameCounters.Add( LastSecondFrameCounter );
-			AccumulatedFrameCounters.Add( GetNumFrames() );
+			FrameCounters.Add(LastSecondFrameCounter);
+			AccumulatedFrameCounters.Add(GetNumFrames());
 		}
 
 		LastSecondFrameTimeMS -= 1000.0f;
@@ -74,12 +78,33 @@ const uint32 IDataProvider::AdvanceFrame( const float DeltaTimeMS )
 	return FrameIndex;
 }
 
+
+const FIntPoint IDataProvider::GetClosestSamplesIndicesForTime(const float StartTimeMS, const float EndTimeMS) const
+{
+	// Find the first sample that start time is less or equal.
+	const int StartIndex = FBinaryFindIndex::LessEqual(ElapsedFrameTimes, StartTimeMS);
+	int32 EndIndex = StartIndex;
+
+	// Find the first sample that start time is greater or equal, usually the next sample. 
+	// More efficient then using binary search again.
+	for (; EndIndex < ElapsedFrameTimes.Num(); ++EndIndex)
+	{
+		const float NextIndexTimeMS = ElapsedFrameTimes[EndIndex];
+		if (NextIndexTimeMS >= EndTimeMS)
+		{
+			break;
+		}
+	}
+
+	return FIntPoint(StartIndex, FMath::Min(ElapsedFrameTimes.Num(), EndIndex + 1));
+}
+
+
 /*-----------------------------------------------------------------------------
 	FArrayDataProvider
 -----------------------------------------------------------------------------*/
 
-const uint32 FArrayDataProvider::AddHierarchicalSample
-( 
+const uint32 FArrayDataProvider::AddHierarchicalSample(
 	const uint32 InThreadID, 
 	const uint32 InGroupID, 
 	const uint32 InStatID, 
@@ -89,15 +114,15 @@ const uint32 FArrayDataProvider::AddHierarchicalSample
 )
 {
 	uint32 HierSampleIndex = Collection.Num();
-	FProfilerSample HierarchicalSample( InThreadID, InGroupID, InStatID, InDurationCycles, InCallsPerFrame );
+	FProfilerSample HierarchicalSample(InThreadID, InGroupID, InStatID, InDurationCycles, InCallsPerFrame);
 
-	Collection.AddElement( HierarchicalSample );
+	Collection.AddElement(HierarchicalSample);
 
-	if( FProfilerSample::IsIndexValid( InParentIndex ) )
+	if(FProfilerSample::IsIndexValid(InParentIndex))
 	{
 		FProfilerSample& Parent = Collection[ InParentIndex ];
 		const uint32 InitialMemUsage = Parent.ChildrenIndices().GetAllocatedSize();
-		Parent.AddChild( HierSampleIndex );
+		Parent.AddChild(HierSampleIndex);
 		const uint32 AccumulatedMemoryUsage = Parent.ChildrenIndices().GetAllocatedSize() - InitialMemUsage;
 		ChildrenIndicesMemoryUsage += AccumulatedMemoryUsage;
 	}
@@ -105,30 +130,33 @@ const uint32 FArrayDataProvider::AddHierarchicalSample
 	return HierSampleIndex;
 }
 
-void FArrayDataProvider::AddCounterSample
-( 
+
+void FArrayDataProvider::AddCounterSample(
 	const uint32 InGroupID, 
 	const uint32 InStatID, 
 	const double InCounter, 
-	const EProfilerSampleTypes::Type InProfilerSampleType 
+	const EProfilerSampleTypes::Type InProfilerSampleType
 )
 {
-	FProfilerSample CounterSample( InGroupID, InStatID, InCounter, InProfilerSampleType );
-	Collection.AddElement( CounterSample );
+	FProfilerSample CounterSample(InGroupID, InStatID, InCounter, InProfilerSampleType);
+	Collection.AddElement(CounterSample);
 }
 
-const uint32 FArrayDataProvider::AddDuplicatedSample( const FProfilerSample& ProfilerSample )
+
+const uint32 FArrayDataProvider::AddDuplicatedSample(const FProfilerSample& ProfilerSample)
 {
 	uint32 DuplicateSampleIndex = Collection.Num();
-	Collection.AddElement( ProfilerSample );
+	Collection.AddElement(ProfilerSample);
 	return DuplicateSampleIndex;
 }
 
+
 const uint32 FArrayDataProvider::GetNumSamples() const
 {
 	return Collection.Num();
 }
 
+
 const SIZE_T FArrayDataProvider::GetMemoryUsage() const
 {
 	SIZE_T MemoryUsage = FrameTimes.GetAllocatedSize();
diff --git a/Engine/Source/Developer/Profiler/Private/ProfilerDataProvider.h b/Engine/Source/Developer/Profiler/Private/ProfilerDataProvider.h
--- a/Engine/Source/Developer/Profiler/Private/ProfilerDataProvider.h
+++ b/Engine/Source/Developer/Profiler/Private/ProfilerDataProvider.h
@@ -2,17 +2,25 @@
 
 #pragma once
 
+#include "ProfilerSample.h"
+
+
+class IDataProvider;
+
+
 /*-----------------------------------------------------------------------------
 	Declarations
 -----------------------------------------------------------------------------*/
 
 /** Data provider interface, acts as a container for profiler samples. */
-class IDataProvider : public TSharedFromThis<IDataProvider>
+class IDataProvider
+	: public TSharedFromThis<IDataProvider>
 {
 	friend class FGraphDataSource;
 	friend class FEventGraphData;
 
 public:
+
 	/** Default constructor. */
 	IDataProvider()
 		: ElapsedTimeMS( 0.0f )
@@ -23,14 +31,16 @@ public:
 	{}
 
 	/** Virtual destructor. */
-	virtual ~IDataProvider()
-	{}
+	virtual ~IDataProvider() { }
 
 protected:
+
 	/**
 	 * Method for creating a duplicated copy of this data provider instance with a particular group of frames.
+	 *
 	 * Data is stored in the instance of FArrayDataProvider class.
 	 * Data source for the SEventGraph widget.
+	 *
 	 * @return A shared reference with a new instance of the data provider.
 	 */
 	template< class TDataProviderType >
@@ -45,7 +55,7 @@ protected:
 		const uint32 FrameEndIndex = FrameStartIndex + NumFrames;
 		check( FrameEndIndex <= TotalNumFrames );
 
-		IDataProviderRef DataProvider = MakeShareable( new TDataProviderType );
+		TSharedRef<IDataProvider> DataProvider = MakeShareable( new TDataProviderType );
 		InternalDuplicate( DataProvider, FrameStartIndex, FrameEndIndex );
 
 		return DataProvider;
@@ -53,19 +63,19 @@ protected:
 
 protected:
 	/** Helper method used to create a copy of a specified data provider. Removes the template parameters, so can be implemented in the source file. */
-	void InternalDuplicate( const IDataProviderRef& DataProvider, const uint32 FrameStartIndex, const uint32 FrameEndIndex ) const;
+	void InternalDuplicate( const TSharedRef<IDataProvider>& DataProvider, const uint32 FrameStartIndex, const uint32 FrameEndIndex ) const;
 
 public:
 	/**
 	 * Adds a new hierarchical sample to the data provider.
 	 *
-	 * @param InThreadID		- The ID of the thread that this profiler sample was captured on
-	 * @param InGroupID			- The ID of the stat group that this profiler sample belongs to
-	 * @param InStatID			- The ID of the stat of this profiler sample
-	 * @param InStartMS			- The start time of this profiler sample, in milliseconds
-	 * @param InDurationMS		- The duration of this profiler sample, in milliseconds
-	 * @param InCallsPerFrame	- The number of times this counter was called in the frame it was captured in
-	 * @param InParentIndex		- The parent of this profiler sample, as an index to a profiler sample
+	 * @param InThreadID The ID of the thread that this profiler sample was captured on
+	 * @param InGroupID The ID of the stat group that this profiler sample belongs to
+	 * @param InStatID The ID of the stat of this profiler sample
+	 * @param InStartMS The start time of this profiler sample, in milliseconds
+	 * @param InDurationMS The duration of this profiler sample, in milliseconds
+	 * @param InCallsPerFrame The number of times this counter was called in the frame it was captured in
+	 * @param InParentIndex The parent of this profiler sample, as an index to a profiler sample
 	 *
 	 */
 	virtual const uint32 AddHierarchicalSample
@@ -81,10 +91,10 @@ public:
 	/**
 	 * Adds a new non-hierarchical sample to the data provider.
 	 *
-	 * @param InGroupID				- The ID of the stat group that this profiler sample belongs to
-	 * @param InStatID				- The ID of the stat of this profiler sample
-	 * @param InCounter				- The counter value for this profiler sample
-	 * @param InProfilerSampleType	- The profiler sample type of this profiler sample
+	 * @param InGroupID The ID of the stat group that this profiler sample belongs to
+	 * @param InStatID The ID of the stat of this profiler sample
+	 * @param InCounter The counter value for this profiler sample
+	 * @param InProfilerSampleType The profiler sample type of this profiler sample
 	 *
 	 */
 	virtual void AddCounterSample
@@ -96,10 +106,11 @@ public:
 	) = 0;
 
 protected:
+
 	/**  
 	 * Adds a sample to the data provider.
 	 * 
-	 * @return	An index to the newly created profiler sample.
+	 * @return An index to the newly created profiler sample.
 	 */
 	virtual const uint32 AddDuplicatedSample
 	( 
@@ -107,6 +118,7 @@ protected:
 	) = 0;
 
 public:
+
 	/**
 	 * @return a reference to the collection where all the profiler samples are stored.
 	 */
@@ -142,25 +154,7 @@ public:
 		return AccumulatedFrameCounters[ SecondIndex ];
 	}
 
-	const FIntPoint GetClosestSamplesIndicesForTime( const float StartTimeMS, const float EndTimeMS ) const
-	{
-		// Find the first sample that start time is less or equal.
-		const int StartIndex = FBinaryFindIndex::LessEqual( ElapsedFrameTimes, StartTimeMS );
-		int32 EndIndex = StartIndex;
-
-		// Find the first sample that start time is greater or equal, usually the next sample. 
-		// More efficient then using binary search again.
-		for( ; EndIndex < ElapsedFrameTimes.Num(); ++EndIndex )
-		{
-			const float NextIndexTimeMS = ElapsedFrameTimes[EndIndex];
-			if( NextIndexTimeMS >= EndTimeMS )
-			{
-				break;
-			}
-		}
-
-		return FIntPoint( StartIndex, FMath::Min(ElapsedFrameTimes.Num(),EndIndex+1) );
-	}
+	const FIntPoint GetClosestSamplesIndicesForTime( const float StartTimeMS, const float EndTimeMS ) const;
 
 	/**
 	 * @return An instance of FIntPoint 
@@ -187,6 +181,7 @@ public:
 	virtual const SIZE_T GetMemoryUsage() const = 0;
 		
 protected:
+
 	/** An array of indices to the frame's range. */
 	TArray<FIntPoint> FrameIndices;
 
@@ -206,6 +201,7 @@ protected:
 	double ElapsedTimeMS;
 
 private:
+
 	/** Accumulates frame times until it reaches a value of one second. */
 	float LastSecondFrameTimeMS;
 
@@ -219,19 +215,23 @@ private:
 	bool bHasAddedSecondStartMarker;
 };
 
+
 /** Implements the data provider interface where samples are stored in the instance of the TArray class. */
-class FArrayDataProvider : public IDataProvider
+class FArrayDataProvider
+	: public IDataProvider
 {
 public:
+
 	FArrayDataProvider()
 		: IDataProvider()
 		, ChildrenIndicesMemoryUsage( 0 )
 	{}
 
-	virtual ~FArrayDataProvider()
-	{
+	virtual ~FArrayDataProvider() { }
 
-	}
+public:
+
+	//~ IDataProvider interface
 
 	virtual const uint32 AddHierarchicalSample
 	( 
@@ -263,6 +263,7 @@ public:
 	}
 
 protected:
+
 	/** Profiler samples stored in TArray. */
 	FProfilerSampleArray Collection;
 
diff --git a/Engine/Source/Developer/Profiler/Private/ProfilerDataSource.cpp b/Engine/Source/Developer/Profiler/Private/ProfilerDataSource.cpp
--- a/Engine/Source/Developer/Profiler/Private/ProfilerDataSource.cpp
+++ b/Engine/Source/Developer/Profiler/Private/ProfilerDataSource.cpp
@@ -1,8 +1,14 @@
 // Copyright 1998-2016 Epic Games, Inc. All Rights Reserved.
 
 #include "ProfilerPrivatePCH.h"
+#include "ProfilerDataProvider.h"
+#include "ProfilerDataSource.h"
+#include "ProfilerSession.h"
+
+
 #define LOCTEXT_NAMESPACE "ProfilerDataSource"
 
+
 /*-----------------------------------------------------------------------------
 	FGraphDataSource
 -----------------------------------------------------------------------------*/
@@ -13,17 +19,18 @@ const int32 FTimeAccuracy::_FPS030 = 30;
 const int32 FTimeAccuracy::_FPS060 = 60;
 const int32 FTimeAccuracy::_FPS120 = 120;
 
+
 /*-----------------------------------------------------------------------------
 	FGraphDataSource
 -----------------------------------------------------------------------------*/
 
-FGraphDataSource::FGraphDataSource( const FProfilerSessionRef& InProfilerSession, const uint32 InStatID ) 
+FGraphDataSource::FGraphDataSource( const TSharedRef<FProfilerSession>& InProfilerSession, const uint32 InStatID ) 
 	: FGraphDataSourceDescription( InStatID )
 	, ThisCachedDataByIndex()
 	, ThisCachedDataByTime( FTimeAccuracy::FPS060 )
 	, ProfilerSession( InProfilerSession ) 
 {
-	const FProfilerStatMetaDataRef MetaData = ProfilerSession->GetMetaData();
+	const TSharedRef<FProfilerStatMetaData> MetaData = ProfilerSession->GetMetaData();
 	const FProfilerStat& Stat = MetaData->GetStatByID(InStatID);
 	const FProfilerGroup& Group = Stat.OwningGroup();
 
@@ -64,7 +71,7 @@ const TGraphDataType FGraphDataSource::GetUncachedValueFromIndex( const uint32 F
 	}
 	else
 	{
-		const IDataProviderRef DataProvider = ProfilerSession->GetDataProvider();
+		const TSharedRef<IDataProvider> DataProvider = ProfilerSession->GetDataProvider();
 
 		const FIntPoint& IndicesForFrame = DataProvider->GetSamplesIndicesForFrame( FrameIndex );
 		const uint32 SampleStartIndex = IndicesForFrame.X;
@@ -90,7 +97,7 @@ const TGraphDataType FGraphDataSource::GetUncachedValueFromIndex( const uint32 F
 
 const TGraphDataType FGraphDataSource::GetUncachedValueFromTimeRange( const float StartTimeMS, const float EndTimeMS ) const
 {
-	const IDataProviderRef DataProvider = ProfilerSession->GetDataProvider();
+	const TSharedRef<IDataProvider> DataProvider = ProfilerSession->GetDataProvider();
 	const FIntPoint IndicesForFrame = DataProvider->GetClosestSamplesIndicesForTime( StartTimeMS, EndTimeMS );
 
 	const uint32 StartFrameIndex = IndicesForFrame.X;
@@ -118,7 +125,7 @@ const float FGraphDataSource::GetTotalTimeMS() const
 	return ProfilerSession->GetDataProvider()->GetTotalTimeMS();
 }
 
-const IDataProviderRef FGraphDataSource::GetDataProvider() const
+const TSharedRef<IDataProvider> FGraphDataSource::GetDataProvider() const
 {
 	return ProfilerSession->GetDataProvider();
 }
@@ -133,11 +140,11 @@ const FGuid FGraphDataSource::GetSessionInstanceID() const
 	return ProfilerSession->GetInstanceID();
 }
 
+
 /*-----------------------------------------------------------------------------
 	FCombinedGraphDataSource
 -----------------------------------------------------------------------------*/
 
-
 FCombinedGraphDataSource::FCombinedGraphDataSource( const uint32 InStatID, const FTimeAccuracy::Type InTimeAccuracy ) 
 	: FGraphDataSourceDescription( InStatID )
 	, ThisCachedDataByTime( InTimeAccuracy )
@@ -154,7 +161,7 @@ const FVector FCombinedGraphDataSource::GetUncachedValueFromTimeRange( const flo
 
 	for( auto It = GetSourcesIterator(); It; ++It )
 	{
-		const FGraphDataSourceRefConst& GraphDataSource = It.Value();
+		const TSharedRef<const FGraphDataSource>& GraphDataSource = It.Value();
 		const float DataSourceValue = GraphDataSource->GetValueFromTimeRange( StartTimeMS, EndTimeMS );
 
 		AggregatedValue.X = FMath::Min( AggregatedValue.X, DataSourceValue );
@@ -170,7 +177,7 @@ void FCombinedGraphDataSource::GetStartIndicesFromTimeRange( const float StartTi
 {
 	for( auto It = GetSourcesIterator(); It; ++It )
 	{
-		const FGraphDataSourceRefConst& GraphDataSource = It.Value();
+		const TSharedRef<const FGraphDataSource>& GraphDataSource = It.Value();
 		const FIntPoint IndicesForFrame = GraphDataSource->GetDataProvider()->GetClosestSamplesIndicesForTime( StartTimeMS, EndTimeMS );
 
 		const uint32 StartFrameIndex = IndicesForFrame.X;
@@ -197,6 +204,7 @@ void FCombinedGraphDataSource::GetStartIndicesFromTimeRange( const float StartTi
 	}
 }
 
+
 /*-----------------------------------------------------------------------------
 	FEventGraphData
 -----------------------------------------------------------------------------*/
@@ -207,6 +215,7 @@ FName FEventGraphConsts::RootEvent = TEXT("RootEvent");
 FName FEventGraphConsts::Self = TEXT("Self");
 FName FEventGraphConsts::FakeRoot = TEXT("FakeRoot");
 
+
 /*-----------------------------------------------------------------------------
 	FEventGraphSample
 -----------------------------------------------------------------------------*/
@@ -415,8 +424,8 @@ FEventGraphData::FEventGraphData( const FProfilerSession * const InProfilerSessi
 	Description = FString::Printf( TEXT("%s: %i"), *InProfilerSession->GetShortName(), InFrameIndex );
 
 	// @TODO: Duplicate is not needed, remove it later.
-	const IDataProviderRef& SessionDataProvider = InProfilerSession->GetDataProvider(); 
-	const IDataProviderRef DataProvider = SessionDataProvider->Duplicate<FArrayDataProvider>( FrameStartIndex );
+	const TSharedRef<IDataProvider>& SessionDataProvider = InProfilerSession->GetDataProvider(); 
+	const TSharedRef<IDataProvider> DataProvider = SessionDataProvider->Duplicate<FArrayDataProvider>( FrameStartIndex );
 
 	const double FrameDurationMS = DataProvider->GetFrameTimeMS( 0 ); 
 	const FProfilerSample& RootProfilerSample = DataProvider->GetCollection()[0];
@@ -426,7 +435,7 @@ FEventGraphData::FEventGraphData( const FProfilerSession * const InProfilerSessi
 	PopulateHierarchy_Recurrent( InProfilerSession, RootEvent, RootProfilerSample, DataProvider );
 
 	// Root sample contains FrameDurationMS
-	const FProfilerStatMetaDataRef& MetaData = InProfilerSession->GetMetaData();
+	const TSharedRef<FProfilerStatMetaData>& MetaData = InProfilerSession->GetMetaData();
 	RootEvent->_InclusiveTimeMS = MetaData->ConvertCyclesToMS( RootProfilerSample.GetDurationCycles() );
 	RootEvent->_MaxInclusiveTimeMS = RootEvent->_MinInclusiveTimeMS = RootEvent->_AvgInclusiveTimeMS = RootEvent->_InclusiveTimeMS;
 	RootEvent->_InclusiveTimePct = 100.0f;
@@ -445,10 +454,10 @@ void FEventGraphData::PopulateHierarchy_Recurrent
 	const FProfilerSession * const ProfilerSession,
 	const FEventGraphSamplePtr ParentEvent, 
 	const FProfilerSample& ParentSample, 
-	const IDataProviderRef DataProvider
+	const TSharedRef<IDataProvider> DataProvider
 )
 {
-	const FProfilerStatMetaDataRef& MetaData = ProfilerSession->GetMetaData();
+	const TSharedRef<FProfilerStatMetaData>& MetaData = ProfilerSession->GetMetaData();
 
 	for( int32 ChildIndex = 0; ChildIndex < ParentSample.ChildrenIndices().Num(); ChildIndex++ )
 	{
