Sequencer: Fix component binding so an existing component binding is used if found.

diff --git a/Engine/Plugins/MovieScene/MatineeToLevelSequence/Source/MatineeToLevelSequence/Private/MatineeToLevelSequenceModule.cpp b/Engine/Plugins/MovieScene/MatineeToLevelSequence/Source/MatineeToLevelSequence/Private/MatineeToLevelSequenceModule.cpp
--- a/Engine/Plugins/MovieScene/MatineeToLevelSequence/Source/MatineeToLevelSequence/Private/MatineeToLevelSequenceModule.cpp
+++ b/Engine/Plugins/MovieScene/MatineeToLevelSequence/Source/MatineeToLevelSequence/Private/MatineeToLevelSequenceModule.cpp
@@ -370,12 +370,18 @@ protected:
 		if (PropObject && Property)
 		{
 			// If the property object that owns this property isn't already bound, add a binding to the property object
-			ObjectGuid = Player.FindObjectId(*PropObject, MovieSceneSequenceID::Root);
+			ObjectGuid = NewSequence->FindPossessableObjectId(*PropObject, InActor);
 			if (!ObjectGuid.IsValid())
 			{
 				UObject* BindingContext = InActor->GetWorld();
 				ObjectGuid = NewMovieScene->AddPossessable(PropObject->GetName(), PropObject->GetClass());
-				NewSequence->BindPossessableObject(ObjectGuid, *PropObject, BindingContext);
+				FMovieScenePossessable* ChildPossessable = NewMovieScene->FindPossessable(ObjectGuid);
+				NewSequence->BindPossessableObject(ObjectGuid, *PropObject, InActor);
+
+				if (ChildPossessable)
+				{
+					ChildPossessable->SetParent(PossessableGuid);
+				}
 			}
 
 			// cbb: String manipulations to get the property path in the right form for sequencer
