Sequencer: Fix to support existing component bindings and camera cuts

diff --git a/Engine/Plugins/MovieScene/MatineeToLevelSequence/Source/MatineeToLevelSequence/Private/MatineeToLevelSequenceModule.cpp b/Engine/Plugins/MovieScene/MatineeToLevelSequence/Source/MatineeToLevelSequence/Private/MatineeToLevelSequenceModule.cpp
--- a/Engine/Plugins/MovieScene/MatineeToLevelSequence/Source/MatineeToLevelSequence/Private/MatineeToLevelSequenceModule.cpp
+++ b/Engine/Plugins/MovieScene/MatineeToLevelSequence/Source/MatineeToLevelSequence/Private/MatineeToLevelSequenceModule.cpp
@@ -357,7 +357,7 @@ protected:
 
 	/** Add property to possessable node **/
 	template <typename T>
-	static T* AddPropertyTrack(FName InPropertyName, AActor* InActor, const FGuid& PossessableGuid, IMovieScenePlayer& Player, UMovieSceneSequence* NewSequence, UMovieScene* NewMovieScene, int32& NumWarnings)
+	static T* AddPropertyTrack(FName InPropertyName, AActor* InActor, const FGuid& PossessableGuid, IMovieScenePlayer& Player, UMovieSceneSequence* NewSequence, UMovieScene* NewMovieScene, int32& NumWarnings, TMap<UObject*, FGuid>& BoundObjectsToGuids)
 	{
 		T* PropertyTrack = nullptr;
 
@@ -370,18 +370,23 @@ protected:
 		if (PropObject && Property)
 		{
 			// If the property object that owns this property isn't already bound, add a binding to the property object
-			ObjectGuid = NewSequence->FindPossessableObjectId(*PropObject, InActor);
-			if (!ObjectGuid.IsValid())
+			if (BoundObjectsToGuids.Contains(PropObject))
+			{
+				ObjectGuid = BoundObjectsToGuids[PropObject];
+			}
+			else
 			{
-				UObject* BindingContext = InActor->GetWorld();
 				ObjectGuid = NewMovieScene->AddPossessable(PropObject->GetName(), PropObject->GetClass());
-				FMovieScenePossessable* ChildPossessable = NewMovieScene->FindPossessable(ObjectGuid);
 				NewSequence->BindPossessableObject(ObjectGuid, *PropObject, InActor);
 
-				if (ChildPossessable)
-				{
-					ChildPossessable->SetParent(PossessableGuid);
-				}
+				BoundObjectsToGuids.Add(PropObject, ObjectGuid);
+			}
+
+			FMovieScenePossessable* ChildPossessable = NewMovieScene->FindPossessable(ObjectGuid);
+
+			if (ChildPossessable)
+			{
+				ChildPossessable->SetParent(PossessableGuid);
 			}
 
 			// cbb: String manipulations to get the property path in the right form for sequencer
@@ -537,6 +542,8 @@ protected:
 	/** Convert an interp group */
 	void ConvertInterpGroup(UInterpGroup* Group, AActor* GroupActor, IMovieScenePlayer& Player, UMovieSceneSequence* NewSequence, UMovieScene* NewMovieScene, int32& NumWarnings)
 	{
+		TMap<UObject*, FGuid> BoundObjectsToGuids;
+
 		FGuid PossessableGuid;
 
 		// Bind the group actor as a possessable						
@@ -643,7 +650,7 @@ protected:
 				UInterpTrackBoolProp* MatineeBoolTrack = StaticCast<UInterpTrackBoolProp*>(Track);
 				if (MatineeBoolTrack->GetNumKeyframes() != 0 && GroupActor && PossessableGuid.IsValid())
 				{
-					UMovieSceneBoolTrack* BoolTrack = AddPropertyTrack<UMovieSceneBoolTrack>(MatineeBoolTrack->PropertyName, GroupActor, PossessableGuid, Player, NewSequence, NewMovieScene, NumWarnings);
+					UMovieSceneBoolTrack* BoolTrack = AddPropertyTrack<UMovieSceneBoolTrack>(MatineeBoolTrack->PropertyName, GroupActor, PossessableGuid, Player, NewSequence, NewMovieScene, NumWarnings, BoundObjectsToGuids);
 					if (BoolTrack)
 					{
 						FMatineeImportTools::CopyInterpBoolTrack(MatineeBoolTrack, BoolTrack);
@@ -655,7 +662,7 @@ protected:
 				UInterpTrackFloatProp* MatineeFloatTrack = StaticCast<UInterpTrackFloatProp*>(Track);
 				if (MatineeFloatTrack->GetNumKeyframes() != 0 && GroupActor && PossessableGuid.IsValid())
 				{
-					UMovieSceneFloatTrack* FloatTrack = AddPropertyTrack<UMovieSceneFloatTrack>(MatineeFloatTrack->PropertyName, GroupActor, PossessableGuid, Player, NewSequence, NewMovieScene, NumWarnings);
+					UMovieSceneFloatTrack* FloatTrack = AddPropertyTrack<UMovieSceneFloatTrack>(MatineeFloatTrack->PropertyName, GroupActor, PossessableGuid, Player, NewSequence, NewMovieScene, NumWarnings, BoundObjectsToGuids);
 					if (FloatTrack)
 					{
 						FMatineeImportTools::CopyInterpFloatTrack(MatineeFloatTrack, FloatTrack);
@@ -699,7 +706,7 @@ protected:
 				UInterpTrackVectorProp* MatineeVectorTrack = StaticCast<UInterpTrackVectorProp*>(Track);
 				if (MatineeVectorTrack->GetNumKeyframes() != 0 && GroupActor && PossessableGuid.IsValid())
 				{
-					UMovieSceneVectorTrack* VectorTrack = AddPropertyTrack<UMovieSceneVectorTrack>(MatineeVectorTrack->PropertyName, GroupActor, PossessableGuid, Player, NewSequence, NewMovieScene, NumWarnings);
+					UMovieSceneVectorTrack* VectorTrack = AddPropertyTrack<UMovieSceneVectorTrack>(MatineeVectorTrack->PropertyName, GroupActor, PossessableGuid, Player, NewSequence, NewMovieScene, NumWarnings, BoundObjectsToGuids);
 					if (VectorTrack)
 					{
 						VectorTrack->SetNumChannelsUsed(3);
@@ -712,7 +719,7 @@ protected:
 				UInterpTrackColorProp* MatineeColorTrack = StaticCast<UInterpTrackColorProp*>(Track);
 				if (MatineeColorTrack->GetNumKeyframes() != 0 && GroupActor && PossessableGuid.IsValid())
 				{
-					UMovieSceneColorTrack* ColorTrack = AddPropertyTrack<UMovieSceneColorTrack>(MatineeColorTrack->PropertyName, GroupActor, PossessableGuid, Player, NewSequence, NewMovieScene, NumWarnings);
+					UMovieSceneColorTrack* ColorTrack = AddPropertyTrack<UMovieSceneColorTrack>(MatineeColorTrack->PropertyName, GroupActor, PossessableGuid, Player, NewSequence, NewMovieScene, NumWarnings, BoundObjectsToGuids);
 					if (ColorTrack)
 					{
 						FMatineeImportTools::CopyInterpColorTrack(MatineeColorTrack, ColorTrack);
@@ -724,7 +731,7 @@ protected:
 				UInterpTrackLinearColorProp* MatineeLinearColorTrack = StaticCast<UInterpTrackLinearColorProp*>(Track);
 				if (MatineeLinearColorTrack->GetNumKeyframes() != 0 && GroupActor && PossessableGuid.IsValid())
 				{
-					UMovieSceneColorTrack* ColorTrack = AddPropertyTrack<UMovieSceneColorTrack>(MatineeLinearColorTrack->PropertyName, GroupActor, PossessableGuid, Player, NewSequence, NewMovieScene, NumWarnings);
+					UMovieSceneColorTrack* ColorTrack = AddPropertyTrack<UMovieSceneColorTrack>(MatineeLinearColorTrack->PropertyName, GroupActor, PossessableGuid, Player, NewSequence, NewMovieScene, NumWarnings, BoundObjectsToGuids);
 					if (ColorTrack)
 					{
 						FMatineeImportTools::CopyInterpLinearColorTrack(MatineeLinearColorTrack, ColorTrack);
@@ -851,6 +858,10 @@ protected:
 				ConvertInterpGroup(Group, GroupActor, TemporaryPlayer, NewSequence, NewMovieScene, NumWarnings);
 			}
 
+			// Force an evaluation so that bound objects will 
+			FMovieSceneContext Context = FMovieSceneContext(FMovieSceneEvaluationRange(FFrameTime(), FFrameRate()), EMovieScenePlayerStatus::Jumping);
+			TemporaryPlayer.GetEvaluationTemplate().Evaluate(Context, TemporaryPlayer);
+
 			// Director group - convert this after the regular groups to ensure that the camera cut bindings are there
 			UInterpGroupDirector* DirGroup = MatineeActor->MatineeData->FindDirectorGroup();
 			if (DirGroup)
