Fix deprecation warnings
#rnx

diff --git a/Engine/Plugins/MovieScene/MatineeToLevelSequence/Source/MatineeToLevelSequence/Private/MatineeToLevelSequenceModule.cpp b/Engine/Plugins/MovieScene/MatineeToLevelSequence/Source/MatineeToLevelSequence/Private/MatineeToLevelSequenceModule.cpp
--- a/Engine/Plugins/MovieScene/MatineeToLevelSequence/Source/MatineeToLevelSequence/Private/MatineeToLevelSequenceModule.cpp
+++ b/Engine/Plugins/MovieScene/MatineeToLevelSequence/Source/MatineeToLevelSequence/Private/MatineeToLevelSequenceModule.cpp
@@ -312,7 +312,7 @@ protected:
 					FText::AsNumber(CameraAnimsToConvert.Num()), FText::AsNumber(NumWarnings));
 			FNotificationInfo NotificationInfo(NotificationText);
 			NotificationInfo.ExpireDuration = 5.f;
-			NotificationInfo.Hyperlink = FSimpleDelegate::CreateStatic([](){ FGlobalTabmanager::Get()->InvokeTab(FName("OutputLog")); });
+			NotificationInfo.Hyperlink = FSimpleDelegate::CreateStatic([](){ FGlobalTabmanager::Get()->TryInvokeTab(FName("OutputLog")); });
 			NotificationInfo.HyperlinkText = LOCTEXT("ShowMessageLogHyperlink", "Show Output Log");
 			FSlateNotificationManager::Get().AddNotification(NotificationInfo);
 		}
