Sequencer: Change UE4 to Unreal Engine

diff --git a/Engine/Plugins/MovieScene/MatineeToLevelSequence/Source/MatineeToLevelSequence/Private/MatineeToLevelSequenceModule.cpp b/Engine/Plugins/MovieScene/MatineeToLevelSequence/Source/MatineeToLevelSequence/Private/MatineeToLevelSequenceModule.cpp
--- a/Engine/Plugins/MovieScene/MatineeToLevelSequence/Source/MatineeToLevelSequence/Private/MatineeToLevelSequenceModule.cpp
+++ b/Engine/Plugins/MovieScene/MatineeToLevelSequence/Source/MatineeToLevelSequence/Private/MatineeToLevelSequenceModule.cpp
@@ -194,7 +194,7 @@ protected:
 
 		// Pop open a dialog asking whether the user to convert and launcher sequencer or no
 		FSuppressableWarningDialog::FSetupInfo Info( 
-			LOCTEXT("MatineeToLevelSequencePrompt", "As of 4.23, Matinee is no longer supported by UE4 and will be removed from the engine in the near future. Once removed, you will no longer be able to run a Matinee or open Matinee.\n\nWould you like to continue opening Matinee or convert your Matinee to a Level Sequence Asset?"), 
+			LOCTEXT("MatineeToLevelSequencePrompt", "As of 4.23, Matinee is no longer supported by Unreal Engine and will be removed from the engine in the near future. Once removed, you will no longer be able to run a Matinee or open Matinee.\n\nWould you like to continue opening Matinee or convert your Matinee to a Level Sequence Asset?"), 
 			LOCTEXT("MatineeToLevelSequenceTitle", "Convert Matinee to Level Sequence Asset"), 
 			TEXT("MatineeToLevelSequence") );
 		Info.ConfirmText = LOCTEXT("MatineeToLevelSequence_ConfirmText", "Open Matinee");
