Sequencer: Auto-create a camera component binding when needed in the Matinee converter

We can encounter Matinee CameraAnims that have some post-process settings on them without having any other camera properties animated. In that case, the camera component binding wouldn't have been created during the normal Matinee data upgrade, and we need to do it ourselves.

diff --git a/Engine/Plugins/MovieScene/MatineeToLevelSequence/Source/MatineeToLevelSequence/Private/CameraAnimToTemplateSequenceConverter.cpp b/Engine/Plugins/MovieScene/MatineeToLevelSequence/Source/MatineeToLevelSequence/Private/CameraAnimToTemplateSequenceConverter.cpp
--- a/Engine/Plugins/MovieScene/MatineeToLevelSequence/Source/MatineeToLevelSequence/Private/CameraAnimToTemplateSequenceConverter.cpp
+++ b/Engine/Plugins/MovieScene/MatineeToLevelSequence/Source/MatineeToLevelSequence/Private/CameraAnimToTemplateSequenceConverter.cpp
@@ -544,12 +544,26 @@ UObject* FCameraAnimToTemplateSequenceConverter::ConvertSingleCameraAnimToTempla
 		}
 	}
 
+	auto EnsureCameraComponentBinding = [&CameraComponentBindingID, CameraTemplate, SpawnableGuid, NewMovieScene, NewSequence]()
+	{
+		if (!CameraComponentBindingID.IsValid())
+		{
+			// We need to create an object binding for the camera component.
+			UCameraComponent* CameraComponent = CameraTemplate->GetCameraComponent();
+			CameraComponentBindingID = NewMovieScene->AddPossessable(CameraComponent->GetName(), CameraComponent->GetClass());
+			FMovieScenePossessable* CameraComponentPossessable = NewMovieScene->FindPossessable(CameraComponentBindingID);
+			CameraComponentPossessable->SetParent(SpawnableGuid);
+			NewSequence->BindPossessableObject(CameraComponentBindingID, *CameraComponent, nullptr);
+		}
+	};
+
 	// Add a track for the post FX blend weight if it wasn't animated.
 	if (CameraAnimToConvert->BasePostProcessBlendWeight != 1.f)
 	{
 		const FName PostProcessBlendWeightPropertyName = GET_MEMBER_NAME_CHECKED(UCameraComponent, PostProcessBlendWeight);
 
-		ensure(CameraComponentBindingID.IsValid());
+		EnsureCameraComponentBinding();
+
 		UMovieSceneFloatTrack* PPBlendWeightTrack = NewMovieScene->FindTrack<UMovieSceneFloatTrack>(CameraComponentBindingID, PostProcessBlendWeightPropertyName);
 		if (!PPBlendWeightTrack)
 		{
@@ -607,6 +621,8 @@ UObject* FCameraAnimToTemplateSequenceConverter::ConvertSingleCameraAnimToTempla
 			continue;
 		}
 
+		EnsureCameraComponentBinding();
+
 		const FName OverridenPropertyClassName = OverridenProperty->GetClass()->GetFName();
 		if (OverridenProperty->IsA<FBoolProperty>())
 		{
